Fix container reuse bug in the Docker Compose demo for Pixel Streaming

diff --git a/Engine/Extras/Containers/Examples/PixelStreaming/run.sh b/Engine/Extras/Containers/Examples/PixelStreaming/run.sh
--- a/Engine/Extras/Containers/Examples/PixelStreaming/run.sh
+++ b/Engine/Extras/Containers/Examples/PixelStreaming/run.sh
@@ -37,4 +37,4 @@ export UNREAL_ENGINE_RELEASE
 export EXTRA_PEERCONNECTION_OPTIONS
 export PUBLIC_IP
 export PWD=$(pwd)
-docker-compose up $COMPOSE_FLAGS
+docker-compose up --force-recreate $COMPOSE_FLAGS
