Branch snapshot for CL 18056524

[CL 18056524 in ue5-main branch]

diff --git a/Engine/Extras/Containers/Dockerfiles/linux/multi-user-server/Dockerfile b/Engine/Extras/Containers/Dockerfiles/linux/multi-user-server/Dockerfile
new file mode 100644
--- /dev/null
+++ b/Engine/Extras/Containers/Dockerfiles/linux/multi-user-server/Dockerfile
@@ -0,0 +1,16 @@
+ARG UNREAL_ENGINE_RELEASE
+FROM ghcr.io/epicgames/unreal-engine:dev-${UNREAL_ENGINE_RELEASE} as source
+
+# Gather the required files for the UnrealMultiUserServer binary and its dependencies
+COPY copy-target.py /tmp/copy-target.py
+RUN python3 /tmp/copy-target.py '[REDACTED_USER_HOME]/UnrealEngine/Engine' '/tmp/gathered/Engine' UnrealMultiUserServer Linux
+
+# Copy the gathered files from the Unreal Engine development image
+FROM gcr.io/distroless/cc-debian10
+COPY --from=source /tmp/gathered/Engine /opt/Engine
+
+# Expose UDP port 9999
+EXPOSE 9999/udp
+
+# Set the UnrealMultiUserServer binary as the container's entrypoint
+ENTRYPOINT ["/opt/Engine/Binaries/Linux/UnrealMultiUserServer"]
diff --git a/Engine/Extras/Containers/Dockerfiles/linux/runtime-pixel-streaming/Dockerfile b/Engine/Extras/Containers/Dockerfiles/linux/runtime-pixel-streaming/Dockerfile
new file mode 100644
--- /dev/null
+++ b/Engine/Extras/Containers/Dockerfiles/linux/runtime-pixel-streaming/Dockerfile
@@ -0,0 +1,4 @@
+FROM adamrehn/ue4-runtime:18.04-cudagl11.1.1
+
+# Enable the NVIDIA driver capabilities required by the NVENC video encoding API
+ENV NVIDIA_DRIVER_CAPABILITIES ${NVIDIA_DRIVER_CAPABILITIES},video
diff --git a/Engine/Extras/Containers/Dockerfiles/linux/runtime/Dockerfile b/Engine/Extras/Containers/Dockerfiles/linux/runtime/Dockerfile
new file mode 100644
--- /dev/null
+++ b/Engine/Extras/Containers/Dockerfiles/linux/runtime/Dockerfile
@@ -0,0 +1 @@
+FROM adamrehn/ue4-runtime:18.04-vulkan-noaudio
diff --git a/Engine/Extras/Containers/Dockerfiles/linux/signalling-server/Dockerfile b/Engine/Extras/Containers/Dockerfiles/linux/signalling-server/Dockerfile
new file mode 100644
--- /dev/null
+++ b/Engine/Extras/Containers/Dockerfiles/linux/signalling-server/Dockerfile
@@ -0,0 +1,19 @@
+ARG UNREAL_ENGINE_RELEASE
+FROM ghcr.io/epicgames/unreal-engine:dev-${UNREAL_ENGINE_RELEASE} as source
+FROM node:10
+
+# Copy the signalling server source code from the Unreal Engine development image
+COPY --from=source [REDACTED_USER_HOME]/UnrealEngine/Engine/Source/Programs/PixelStreaming/WebServers/SignallingWebServer /opt/SignallingWebServer
+
+# Install the dependencies for the signalling server
+WORKDIR /opt/SignallingWebServer
+RUN npm install .
+
+# Expose TCP port 80 for player WebSocket connections and web server HTTP access
+EXPOSE 80
+
+# Expose TCP port 8888 for streamer WebSocket connections
+EXPOSE 8888
+
+# Set the signalling server as the container's entrypoint
+ENTRYPOINT ["/usr/local/bin/node", "/opt/SignallingWebServer/cirrus.js"]
diff --git a/Engine/Extras/Containers/Dockerfiles/windows/runtime/Dockerfile b/Engine/Extras/Containers/Dockerfiles/windows/runtime/Dockerfile
new file mode 100644
--- /dev/null
+++ b/Engine/Extras/Containers/Dockerfiles/windows/runtime/Dockerfile
@@ -0,0 +1,52 @@
+# escape=`
+ARG BASETAG
+FROM mcr.microsoft.com/windows:${BASETAG} AS full
+
+# Gather the system DLLs that we need from the full Windows base image
+RUN xcopy /y C:\Windows\System32\avicap32.dll C:\GatheredDlls\ && `
+	xcopy /y C:\Windows\System32\avrt.dll C:\GatheredDlls\ && `
+	xcopy /y C:\Windows\System32\d3d10warp.dll C:\GatheredDlls\ && `
+	xcopy /y C:\Windows\System32\D3DSCache.dll C:\GatheredDlls\ && `
+	xcopy /y C:\Windows\System32\dsound.dll C:\GatheredDlls\ && `
+	xcopy /y C:\Windows\System32\dxva2.dll C:\GatheredDlls\ && `
+	xcopy /y C:\Windows\System32\glu32.dll C:\GatheredDlls\ && `
+	xcopy /y C:\Windows\System32\mf.dll C:\GatheredDlls\ && `
+	xcopy /y C:\Windows\System32\mfplat.dll C:\GatheredDlls\ && `
+	xcopy /y C:\Windows\System32\mfplay.dll C:\GatheredDlls\ && `
+	xcopy /y C:\Windows\System32\mfreadwrite.dll C:\GatheredDlls\ && `
+	xcopy /y C:\Windows\System32\msdmo.dll C:\GatheredDlls\ && `
+	xcopy /y C:\Windows\System32\msvfw32.dll C:\GatheredDlls\ && `
+	xcopy /y C:\Windows\System32\opengl32.dll C:\GatheredDlls\ && `
+	xcopy /y C:\Windows\System32\ResourcePolicyClient.dll C:\GatheredDlls\
+
+# Retrieve the DirectX runtime files required by the Unreal Engine, since even the full Windows base image does not include them
+RUN curl --progress -L "https://download.microsoft.com/download/8/4/A/84A35BF1-DAFE-4AE8-82AF-AD2AE20B6B14/directx_Jun2010_redist.exe" --output %TEMP%\directx_redist.exe && `
+	start /wait %TEMP%\directx_redist.exe /Q /T:%TEMP%\DirectX && `
+	expand %TEMP%\DirectX\APR2007_xinput_x64.cab -F:xinput1_3.dll C:\GatheredDlls\ && `
+	expand %TEMP%\DirectX\Feb2010_X3DAudio_x64.cab -F:X3DAudio1_7.dll C:\GatheredDlls\ && `
+	expand %TEMP%\DirectX\Jun2010_D3DCompiler_43_x64.cab -F:D3DCompiler_43.dll C:\GatheredDlls\ && `
+	expand %TEMP%\DirectX\Jun2010_XAudio_x64.cab -F:XAudio2_7.dll C:\GatheredDlls\ && `
+	expand %TEMP%\DirectX\Jun2010_XAudio_x64.cab -F:XAPOFX1_5.dll C:\GatheredDlls\
+
+# Retrieve the DirectX shader compiler files needed for DirectX Raytracing (DXR)
+RUN curl --progress -L "[REDACTED_PROFILE_URL]/DirectXShaderCompiler/releases/download/v1.6.2104/dxc_2021_04-20.zip" --output %TEMP%\dxc.zip && `
+	powershell -Command "Expand-Archive -Path \"$env:TEMP\dxc.zip\" -DestinationPath $env:TEMP" && `
+	xcopy /y %TEMP%\bin\x64\dxcompiler.dll C:\GatheredDlls\ && `
+	xcopy /y %TEMP%\bin\x64\dxil.dll C:\GatheredDlls\
+
+# Copy the required DLLs from the full Windows base image into a smaller Windows Server Core base image
+ARG BASETAG
+FROM mcr.microsoft.com/windows/servercore:${BASETAG}
+COPY --from=full C:\GatheredDlls\ C:\Windows\System32\
+
+# Install the Visual C++ runtime files using Chocolatey
+RUN powershell -NoProfile -ExecutionPolicy Bypass -Command "[System.Net.ServicePointManager]::SecurityProtocol = [System.Net.ServicePointManager]::SecurityProtocol -bor 3072; iex ((New-Object System.Net.WebClient).DownloadString('https://chocolatey.org/install.ps1'))"
+RUN choco install -y vcredist-all
+
+# Copy our entrypoint script for enabling vendor-specific graphics APIs
+COPY entrypoint.cmd C:\entrypoint.cmd
+COPY enable-graphics-apis.ps1 C:\enable-graphics-apis.ps1
+
+# When running containers based on this image, be sure to specify the following flags for the `docker run` command to enable GPU acceleration:
+# --isolation process --device class/5B45201D-F2F2-4F3B-85BB-30FF1F953599
+# (For more information, see: <https://docs.microsoft.com/en-us/virtualization/windowscontainers/deploy-containers/gpu-acceleration>)
diff --git a/Engine/Extras/Containers/Dockerfiles/windows/runtime/entrypoint.cmd b/Engine/Extras/Containers/Dockerfiles/windows/runtime/entrypoint.cmd
new file mode 100755
--- /dev/null
+++ b/Engine/Extras/Containers/Dockerfiles/windows/runtime/entrypoint.cmd
@@ -0,0 +1,7 @@
+@echo off
+
+@rem Enable vendor-specific graphics APIs if the container is running with GPU acceleration
+powershell -ExecutionPolicy Bypass -File "%~dp0.\enable-graphics-apis.ps1"
+
+@rem Run the entrypoint command specified via our command-line parameters
+%*
