Update official container image Dockerfiles for UE5

diff --git a/Engine/Extras/Containers/Dockerfiles/linux/build.sh b/Engine/Extras/Containers/Dockerfiles/linux/build.sh
--- a/Engine/Extras/Containers/Dockerfiles/linux/build.sh
+++ b/Engine/Extras/Containers/Dockerfiles/linux/build.sh
@@ -1,7 +1,7 @@
 #!/usr/bin/env bash
 
 # Determine which release of the Unreal Engine we will be building container images for
-UNREAL_ENGINE_RELEASE="4.27"
+UNREAL_ENGINE_RELEASE="5.0"
 if [[ ! -z "$1" ]]; then
 	UNREAL_ENGINE_RELEASE="$1"
 fi
@@ -20,6 +20,12 @@ if [[ ! -z "$2" ]]; then
 	GIT_REPO="$2"
 fi
 
+# Determine whether the user specified a changelist number to set in Build.version
+CHANGELIST_OVERRIDE=""
+if [[ ! -z "$3" ]]; then
+	CHANGELIST_OVERRIDE="$3"
+fi
+
 # Verify that the user has placed their GitHub username in the file `username.txt`
 if [ ! -f ./username.txt ]; then
 	echo 'Please place your GitHub username in a text file called `username.txt` in the current directory.'
@@ -36,10 +42,10 @@ fi
 
 # Assemble the common arguments to pass to the `docker buildx build` command
 args=(
-	--build-arg 'BASEIMAGE=nvidia/cudagl:11.1.1-devel-ubuntu18.04'
+	--build-arg 'BASEIMAGE=nvidia/opengl:1.0-glvnd-devel-ubuntu18.04'
 	--build-arg "GIT_REPO=${GIT_REPO}"
 	--build-arg "GIT_BRANCH=${GIT_BRANCH}"
-	--build-arg 'BUILD_DDC=true'
+	--build-arg "CHANGELIST=${CHANGELIST_OVERRIDE}"
 	--secret id=username,src=username.txt
 	--secret id=password,src=password.txt
 	--progress=plain
@@ -65,8 +71,6 @@ docker build -t 'ghcr.io/epicgames/unreal-engine:runtime-pixel-streaming' ./runt
 echo '[build.sh] Building the `unreal-engine` image for Unreal Engine release' "${UNREAL_ENGINE_RELEASE}..."
 docker buildx build \
 	-t "ghcr.io/epicgames/unreal-engine:dev-${UNREAL_ENGINE_RELEASE}" \
-	--build-arg 'EXCLUDE_DEBUG=0' \
-	--build-arg 'EXCLUDE_TEMPLATES=0' \
 	"${args[@]}" \
 	./dev
 
@@ -74,10 +78,8 @@ docker buildx build \
 echo '[build.sh] Building the `unreal-engine` image for Unreal Engine release' "${UNREAL_ENGINE_RELEASE}..."
 docker buildx build \
 	-t "ghcr.io/epicgames/unreal-engine:dev-slim-${UNREAL_ENGINE_RELEASE}" \
-	--build-arg 'EXCLUDE_DEBUG=1' \
-	--build-arg 'EXCLUDE_TEMPLATES=1' \
 	"${args[@]}" \
-	./dev
+	./dev-slim
 
 # Build the container image for the Pixel Streaming signalling server, which pulls files from our other images
 docker build \
diff --git a/Engine/Extras/Containers/Dockerfiles/linux/dev-slim/git-credential-helper-secrets.sh b/Engine/Extras/Containers/Dockerfiles/linux/dev-slim/git-credential-helper-secrets.sh
new file mode 100644
--- /dev/null
+++ b/Engine/Extras/Containers/Dockerfiles/linux/dev-slim/git-credential-helper-secrets.sh
@@ -0,0 +1,6 @@
+#!/usr/bin/env bash
+if [[ "$1" == "Password for"* ]]; then
+	cat /run/secrets/password
+else
+	cat /run/secrets/username
+fi
diff --git a/Engine/Extras/Containers/Dockerfiles/linux/dev-slim/set-changelist.py b/Engine/Extras/Containers/Dockerfiles/linux/dev-slim/set-changelist.py
new file mode 100644
--- /dev/null
+++ b/Engine/Extras/Containers/Dockerfiles/linux/dev-slim/set-changelist.py
@@ -0,0 +1,27 @@
+#!/usr/bin/env python3
+import json, sys
+
+
+def readFile(filename):
+    with open(filename, "rb") as f:
+        return f.read().decode("utf-8")
+
+
+def writeFile(filename, data):
+    with open(filename, "wb") as f:
+        f.write(data.encode("utf-8"))
+
+
+# Update the `Changelist` field to reflect the `CompatibleChangelist` field in our version file, unless a specific value was provided
+versionFile = sys.argv[1]
+changelistOverride = int(sys.argv[2]) if len(sys.argv) > 2 else None
+details = json.loads(readFile(versionFile))
+details["Changelist"] = (
+    changelistOverride
+    if changelistOverride is not None
+    else details["CompatibleChangelist"]
+)
+details["IsPromotedBuild"] = 1
+patchedJson = json.dumps(details, indent=4)
+writeFile(versionFile, patchedJson)
+print("PATCHED BUILD.VERSION:\n{}".format(patchedJson), file=sys.stderr)
diff --git a/Engine/Extras/Containers/Dockerfiles/linux/dev-slim/split-components.py b/Engine/Extras/Containers/Dockerfiles/linux/dev-slim/split-components.py
new file mode 100644
--- /dev/null
+++ b/Engine/Extras/Containers/Dockerfiles/linux/dev-slim/split-components.py
@@ -0,0 +1,61 @@
+#!/usr/bin/env python3
+import glob, os, shutil, sys
+from os.path import basename, dirname, exists, join
+
+# Logs a message to stderr
+def log(message):
+    print(message, file=sys.stderr)
+    sys.stderr.flush()
+
+
+# Extracts the files and directories for the specified component and moves them to a separate output directory
+def extractComponent(inputDir, outputDir, component, description, items):
+
+    # Print progress output
+    log("\nExtracting {}...".format(description))
+
+    # Create the output directory for the component if it doesn't already exist
+    componentDir = join(outputDir, component)
+    os.makedirs(outputDir, exist_ok=True)
+
+    # Move each file and directory for the component to the output directory
+    for item in items:
+
+        # Verify that the item exists
+        if not exists(item):
+            log("Skipping non-existent item: {}".format(item))
+            continue
+
+        # Print progress output
+        log("Moving: {}".format(item))
+
+        # Ensure the parent directory for the item exists in the output directory
+        parent = dirname(item).replace(inputDir, componentDir)
+        os.makedirs(parent, exist_ok=True)
+
+        # Perform the move
+        shutil.move(item, join(parent, basename(item)))
+
+
+# Retrieve the path to the root directory of the Installed Build
+rootDir = sys.argv[1]
+
+# Retrieve the path to the root output directory for extracted components and ensure it exists
+outputDir = sys.argv[2]
+os.makedirs(outputDir, exist_ok=True)
+
+# Extract the DDC
+ddc = [join(rootDir, "Engine", "DerivedDataCache", "Compressed.ddp")]
+extractComponent(rootDir, outputDir, "DDC", "Derived Data Cache (DDC)", ddc)
+
+# Extract debug symbols
+symbolFiles = glob.glob(join(rootDir, "**", "*.debug"), recursive=True) + glob.glob(
+    join(rootDir, "**", "*.sym"), recursive=True
+)
+extractComponent(rootDir, outputDir, "DebugSymbols", "debug symbols", symbolFiles)
+
+# Extract template projects and samples
+subdirs = [join(rootDir, subdir) for subdir in ["FeaturePacks", "Samples", "Templates"]]
+extractComponent(
+    rootDir, outputDir, "TemplatesAndSamples", "template projects and samples", subdirs
+)
diff --git a/Engine/Extras/Containers/Dockerfiles/linux/dev/set-changelist.py b/Engine/Extras/Containers/Dockerfiles/linux/dev/set-changelist.py
--- a/Engine/Extras/Containers/Dockerfiles/linux/dev/set-changelist.py
+++ b/Engine/Extras/Containers/Dockerfiles/linux/dev/set-changelist.py
@@ -1,17 +1,27 @@
 #!/usr/bin/env python3
-import json, os, sys
+import json, sys
+
 
 def readFile(filename):
-	with open(filename, 'rb') as f:
-		return f.read().decode('utf-8')
+    with open(filename, "rb") as f:
+        return f.read().decode("utf-8")
+
 
 def writeFile(filename, data):
-	with open(filename, 'wb') as f:
-		f.write(data.encode('utf-8'))
+    with open(filename, "wb") as f:
+        f.write(data.encode("utf-8"))
+
 
-# Update the `Changelist` field to reflect the `CompatibleChangelist` field in our version file
+# Update the `Changelist` field to reflect the `CompatibleChangelist` field in our version file, unless a specific value was provided
 versionFile = sys.argv[1]
+changelistOverride = int(sys.argv[2]) if len(sys.argv) > 2 else None
 details = json.loads(readFile(versionFile))
-details['Changelist'] = details['CompatibleChangelist']
-details['IsPromotedBuild'] = 1
-writeFile(versionFile, json.dumps(details, indent=4))
+details["Changelist"] = (
+    changelistOverride
+    if changelistOverride is not None
+    else details["CompatibleChangelist"]
+)
+details["IsPromotedBuild"] = 1
+patchedJson = json.dumps(details, indent=4)
+writeFile(versionFile, patchedJson)
+print("PATCHED BUILD.VERSION:\n{}".format(patchedJson), file=sys.stderr)
diff --git a/Engine/Extras/Containers/Dockerfiles/linux/dev/split-components.py b/Engine/Extras/Containers/Dockerfiles/linux/dev/split-components.py
--- a/Engine/Extras/Containers/Dockerfiles/linux/dev/split-components.py
+++ b/Engine/Extras/Containers/Dockerfiles/linux/dev/split-components.py
@@ -4,39 +4,38 @@ from os.path import basename, dirname, exists, join
 
 # Logs a message to stderr
 def log(message):
-	print(message, file=sys.stderr)
-	sys.stderr.flush()
+    print(message, file=sys.stderr)
+    sys.stderr.flush()
+
 
 # Extracts the files and directories for the specified component and moves them to a separate output directory
 def extractComponent(inputDir, outputDir, component, description, items):
-	
-	# Print progress output
-	log('\nExtracting {}...'.format(description))
-	
-	# Create the output directory for the component if it doesn't already exist
-	componentDir = join(outputDir, component)
-	os.makedirs(outputDir, exist_ok=True)
-	
-	# Move each file and directory for the component to the output directory
-	for item in items:
-		
-		# Verify that the item exists
-		if not exists(item):
-			log('Skipping non-existent item: {}'.format(item))
-			continue
-		
-		# Print progress output
-		log('Moving: {}'.format(item))
-		
-		# Ensure the parent directory for the item exists in the output directory
-		parent = dirname(item).replace(inputDir, componentDir)
-		os.makedirs(parent, exist_ok=True)
-		
-		# Perform the move
-		shutil.move(
-			item,
-			join(parent, basename(item))
-		)
+
+    # Print progress output
+    log("\nExtracting {}...".format(description))
+
+    # Create the output directory for the component if it doesn't already exist
+    componentDir = join(outputDir, component)
+    os.makedirs(outputDir, exist_ok=True)
+
+    # Move each file and directory for the component to the output directory
+    for item in items:
+
+        # Verify that the item exists
+        if not exists(item):
+            log("Skipping non-existent item: {}".format(item))
+            continue
+
+        # Print progress output
+        log("Moving: {}".format(item))
+
+        # Ensure the parent directory for the item exists in the output directory
+        parent = dirname(item).replace(inputDir, componentDir)
+        os.makedirs(parent, exist_ok=True)
+
+        # Perform the move
+        shutil.move(item, join(parent, basename(item)))
+
 
 # Retrieve the path to the root directory of the Installed Build
 rootDir = sys.argv[1]
@@ -46,13 +45,17 @@ outputDir = sys.argv[2]
 os.makedirs(outputDir, exist_ok=True)
 
 # Extract the DDC
-ddc = [join(rootDir, 'Engine', 'DerivedDataCache', 'Compressed.ddp')]
-extractComponent(rootDir, outputDir, 'DDC', 'Derived Data Cache (DDC)', ddc)
+ddc = [join(rootDir, "Engine", "DerivedDataCache", "Compressed.ddp")]
+extractComponent(rootDir, outputDir, "DDC", "Derived Data Cache (DDC)", ddc)
 
 # Extract debug symbols
-symbolFiles = glob.glob(join(rootDir, '**', '*.debug'), recursive=True) + glob.glob(join(rootDir, '**', '*.sym'), recursive=True)
-extractComponent(rootDir, outputDir, 'DebugSymbols', 'debug symbols', symbolFiles)
+symbolFiles = glob.glob(join(rootDir, "**", "*.debug"), recursive=True) + glob.glob(
+    join(rootDir, "**", "*.sym"), recursive=True
+)
+extractComponent(rootDir, outputDir, "DebugSymbols", "debug symbols", symbolFiles)
 
 # Extract template projects and samples
-subdirs = [join(rootDir, subdir) for subdir in ['FeaturePacks', 'Samples', 'Templates']]
-extractComponent(rootDir, outputDir, 'TemplatesAndSamples', 'template projects and samples', subdirs)
+subdirs = [join(rootDir, subdir) for subdir in ["FeaturePacks", "Samples", "Templates"]]
+extractComponent(
+    rootDir, outputDir, "TemplatesAndSamples", "template projects and samples", subdirs
+)
diff --git a/Engine/Extras/Containers/Dockerfiles/linux/generate.sh b/Engine/Extras/Containers/Dockerfiles/linux/generate.sh
--- a/Engine/Extras/Containers/Dockerfiles/linux/generate.sh
+++ b/Engine/Extras/Containers/Dockerfiles/linux/generate.sh
@@ -1,47 +1,64 @@
 #!/usr/bin/env bash
 
-# Create a temporary directory to hold our generated files before we move them to the actual output directory
-# (Note that this is necessary to ensure ue4-docker doesn't clobber our other Dockerfiles, since it truncates its target directory)
-tempDir=$(mktemp -d -t dockerfile-XXXXXXXXXX)
 
-# Assemble our arguments to pass to the `ue4-docker build` command
-args=(
+# Generates the Dockerfile for the specified container image
+function generateDockerfile {
 	
-	# Note that the release number here is actually irrelevant to the generated files, since the repository and branch used to clone the
-	# Engine source code are defined by Docker build arguments that the user sets when building images with the generated Dockerfiles.
-	# The need to specify a release number is purely an artifact of how the `ue4-docker build` command was originally designed prior to
-	# the addition of functionality to generate Dockerfiles rather than always building container images.
-	'4.27.0'
+	# Create a temporary directory to hold our generated files before we move them to the actual output directory
+	# (Note that this is necessary to ensure ue4-docker doesn't clobber our other Dockerfiles, since it truncates its target directory)
+	tempDir=$(mktemp -d -t dockerfile-XXXXXXXXXX)
 	
-	# Ensure that we generate Dockerfiles for Linux container images even if this script is run on a Windows host system (e.g. in git bash)
-	# (Note that although generating the Dockerfiles under Windows is fine, it is NOT recommended that the images actually be built under Windows)
-	--linux
+	# Assemble our arguments to pass to the `ue4-docker build` command
+	args=(
+		
+		# Note that the release number here is actually irrelevant to the generated files, since the repository and branch used to clone the
+		# Engine source code are defined by Docker build arguments that the user sets when building images with the generated Dockerfiles.
+		# The need to specify a release number is purely an artifact of how the `ue4-docker build` command was originally designed prior to
+		# the addition of functionality to generate Dockerfiles rather than always building container images.
+		'5.0.0'
+		
+		# Ensure that we generate Dockerfiles for Linux container images even if this script is run on a Windows host system (e.g. in git bash)
+		# (Note that although generating the Dockerfiles under Windows is fine, it is NOT recommended that the images actually be built under Windows)
+		--linux
+		
+		# This will ensure our target directory exists and is empty, and will then copy the generated Dockerfiles to it
+		-layout "$tempDir"
+		
+		# This will combine the generated Dockerfiles into a single Dockerfile that uses a multi-stage build
+		--combine
+		
+		# These flags prevent Dockerfiles from being generated for the `ue4-engine` and `ue4-full` images, which we are not interested in
+		--no-engine --no-full
+		
+		# This disables building the Engine for AArch64 when creating an Installed Build, and enables support for client and dedicated server targets
+		--opt buildgraph-args='-set:WithLinuxArm64=false -set:WithClient=true -set:WithServer=true'
+		
+		# This enables the use of BuildKit build secrets, which is necessary in order to build images independently of ue4-docker itself
+		--opt credential-mode=secrets
+		
+		# This strips out the Dockerfile code that ue4-docker would ordinarily generate that is uneccesary for Unreal Engine 4.26+
+		--opt disable-all-patches
+		
+		# This strips out the image labels that are ordinarily used to facilitate cleanup and inspection of images built by ue4-docker
+		--opt disable-labels
+		
+		# This enables the glibc DSO sorting patch, which significantly improves Editor startup times
+		--opt enable-dso-patch
+	)
 	
-	# This will ensure our target directory exists and is empty, and will then copy the generated Dockerfiles to it
-	-layout "$tempDir"
+	# Invoke ue4-docker to generate the Dockerfile
+	ue4-docker build "${args[@]}" $2
 	
-	# This will combine the generated Dockerfiles into a single Dockerfile that uses a multi-stage build
-	--combine
+	# Move the Dockerfile to the actual output directory
+	outputDir="./$1"
+	test -d "$outputDir" && rm -rf "$outputDir"
+	mv "$tempDir/combined" "$outputDir"
 	
-	# These flags prevent Dockerfiles from being generated for the `ue4-engine` and `ue4-full` images, which we are not interested in
-	--no-engine --no-full
-	
-	# This disables building the Engine for AArch64 when creating an Installed Build
-	--opt buildgraph-args='-set:WithLinuxArm64=false'
-	
-	# This enables the use of BuildKit build secrets, which is necessary in order to build images independently of ue4-docker itself
-	--opt credential-mode=secrets
-	
-	# This strips out the Dockerfile code that ue4-docker would ordinarily generate that is uneccesary for Unreal Engine 4.26+
-	--opt disable-all-patches
-	
-	# This strips out the image labels that are ordinarily used to facilitate cleanup and inspection of images built by ue4-docker
-	--opt disable-labels
-)
+	# Clean up the temporary directory
+	rm -rf "$tempDir"
+}
+
 
-# Invoke ue4-docker to generate our Dockerfile, move it to the actual output directory and clean up the temporary directory
-ue4-docker build "${args[@]}"
-outputDir='./dev'
-test -d "$outputDir" && rm -rf "$outputDir"
-mv "$tempDir/combined" "$outputDir"
-rm -rf "$tempDir"
+# Generate the `dev` and `dev-slim` Dockerfiles
+generateDockerfile 'dev'
+generateDockerfile 'dev-slim' '--exclude debug --exclude templates'
diff --git a/Engine/Extras/Containers/Dockerfiles/linux/multi-user-server/copy-target.py b/Engine/Extras/Containers/Dockerfiles/linux/multi-user-server/copy-target.py
--- a/Engine/Extras/Containers/Dockerfiles/linux/multi-user-server/copy-target.py
+++ b/Engine/Extras/Containers/Dockerfiles/linux/multi-user-server/copy-target.py
@@ -82,7 +82,8 @@ target = Utility.parseJson(join(args.EngineDir, 'Binaries', args.Platform, '{}.t
 dependencies = itertools.chain.from_iterable([target['BuildProducts'] + target['RuntimeDependencies']])
 dependencies = [d for d in dependencies if d['Type'] != 'SymbolFile']
 for dependency in dependencies:
-	copier.copy(dependency['Path'])
+	if '$(EngineDir)/Content/Slate' not in dependency['Path']:
+		copier.copy(dependency['Path'])
 
 # If the target is a program then copy any configuration files
 # (This is particularly important for UnrealMultiUserServer since it uses `DefaultEngine.ini` to enable the plugins it needs)
diff --git a/Engine/Extras/Containers/Dockerfiles/linux/signalling-server/Dockerfile b/Engine/Extras/Containers/Dockerfiles/linux/signalling-server/Dockerfile
new file mode 100644
--- /dev/null
+++ b/Engine/Extras/Containers/Dockerfiles/linux/signalling-server/Dockerfile
@@ -0,0 +1,19 @@
+ARG UNREAL_ENGINE_RELEASE
+FROM ghcr.io/epicgames/unreal-engine:dev-${UNREAL_ENGINE_RELEASE} as source
+FROM node:10
+
+# Copy the signalling server source code from the Unreal Engine development image
+COPY --from=source [REDACTED_USER_HOME]/UnrealEngine/Samples/PixelStreaming/WebServers/SignallingWebServer /opt/SignallingWebServer
+
+# Install the dependencies for the signalling server
+WORKDIR /opt/SignallingWebServer
+RUN npm install .
+
+# Expose TCP port 80 for player WebSocket connections and web server HTTP access
+EXPOSE 80
+
+# Expose TCP port 8888 for streamer WebSocket connections
+EXPOSE 8888
+
+# Set the signalling server as the container's entrypoint
+ENTRYPOINT ["/usr/local/bin/node", "/opt/SignallingWebServer/cirrus.js"]
diff --git a/Engine/Extras/Containers/Dockerfiles/linux/test.sh b/Engine/Extras/Containers/Dockerfiles/linux/test.sh
--- a/Engine/Extras/Containers/Dockerfiles/linux/test.sh
+++ b/Engine/Extras/Containers/Dockerfiles/linux/test.sh
@@ -10,6 +10,9 @@ function runTests {
 	# Checkout the source code of a demo Unreal project in the container
 	docker exec 'unrealtest' git clone --progress --depth=1 '[REDACTED_PROFILE_URL]/4.26/firstpersoncxx.git' [REDACTED_USER_HOME]/project
 	
+	# Enable the WebBrowserWidget plugin, to ensure we test the ability to package projects that use CEF
+	docker exec 'unrealtest' sed -i 's|"Modules"|"Plugins": [{"Name": "WebBrowserWidget", "Enabled": true}],\n\t"Modules"|' [REDACTED_USER_HOME]/project/FirstPersonCxx.uproject
+	
 	# Test that we can build the demo Unreal project in the container
 	docker exec 'unrealtest' [REDACTED_USER_HOME]/UnrealEngine/Engine/Build/BatchFiles/Linux/Build.sh \
 		FirstPersonCxxEditor Linux Development \
