Add build script option to create container images from a specific git commit hash

diff --git a/Engine/Extras/Containers/Dockerfiles/linux/build.sh b/Engine/Extras/Containers/Dockerfiles/linux/build.sh
--- a/Engine/Extras/Containers/Dockerfiles/linux/build.sh
+++ b/Engine/Extras/Containers/Dockerfiles/linux/build.sh
@@ -26,6 +26,15 @@ if [[ ! -z "$3" ]]; then
 	CHANGELIST_OVERRIDE="$3"
 fi
 
+# If the user specified a git commit hash instead of a release, branch or tag then automatically determine the changelist
+# number from the commit message, and tag the built container images with an abbreviated version of the commit hash
+COMMIT_HASH=$(echo "${UNREAL_ENGINE_RELEASE}" | grep -E -o '^[0-9a-f]{40}$')
+COMMIT_SHORT=$(echo "${COMMIT_HASH}" | grep -E -o '^[0-9a-f]{8}')
+if [[ ! -z "$COMMIT_HASH" ]]; then
+	UNREAL_ENGINE_RELEASE="$COMMIT_SHORT"
+	CHANGELIST_OVERRIDE="auto"
+fi
+
 # Verify that the user has placed their GitHub username in the file `username.txt`
 if [ ! -f ./username.txt ]; then
 	echo 'Please place your GitHub username in a text file called `username.txt` in the current directory.'
@@ -57,7 +66,8 @@ echo '-------------------'
 echo
 echo "Unreal Engine Release: ${UNREAL_ENGINE_RELEASE}"
 echo "Git Repository:        ${GIT_REPO}"
-echo "Git Branch/Tag:        ${GIT_BRANCH}"
+echo "Git Branch/Tag/Commit: ${GIT_BRANCH}"
+echo "Changelist Override:   ${CHANGELIST_OVERRIDE}"
 echo
 
 # Print commands as they are executed and halt immediately if any command fails
diff --git a/Engine/Extras/Containers/Dockerfiles/linux/dev-slim/Dockerfile b/Engine/Extras/Containers/Dockerfiles/linux/dev-slim/Dockerfile
--- a/Engine/Extras/Containers/Dockerfiles/linux/dev-slim/Dockerfile
+++ b/Engine/Extras/Containers/Dockerfiles/linux/dev-slim/Dockerfile
@@ -1,4 +1,4 @@
-# This file was generated by ue4-docker version 0.0.93 with the following options:
+# This file was generated by ue4-docker version 0.0.94 with the following options:
 # 
 # - buildgraph_args: "-set:WithLinuxArm64=false -set:WithClient=true -set:WithServer=true"
 # - combine: true
@@ -98,7 +98,7 @@ FROM prerequisites as source
 # The git repository that we will clone
 ARG GIT_REPO=""
 
-# The git branch/tag that we will checkout
+# The git branch/tag/commit that we will checkout
 ARG GIT_BRANCH=""
 
 # Install our git credential helper that retrieves credentials from build secrets
@@ -112,7 +112,13 @@ RUN chmod +x /tmp/git-credential-helper-secrets.sh && echo '' && echo 'RUN direc
 ARG CHANGELIST
 RUN --mount=type=secret,id=username,uid=1000,required \
 	--mount=type=secret,id=password,uid=1000,required \
-	CHANGELIST="$CHANGELIST" git clone --progress --depth=1 -b $GIT_BRANCH $GIT_REPO [REDACTED_USER_HOME]/UnrealEngine && echo '' && echo 'RUN directive complete. Docker will now commit the filesystem layer to disk.' && echo 'Note that for large filesystem layers this can take quite some time.' && echo 'Performing filesystem layer commit...' && echo ''
+	CHANGELIST="$CHANGELIST" \
+	mkdir [REDACTED_USER_HOME]/UnrealEngine && \
+	cd [REDACTED_USER_HOME]/UnrealEngine && \
+	git init && \
+	git remote add origin "$GIT_REPO" && \
+	git fetch --progress --depth 1 origin "$GIT_BRANCH" && \
+	git checkout FETCH_HEAD && echo '' && echo 'RUN directive complete. Docker will now commit the filesystem layer to disk.' && echo 'Note that for large filesystem layers this can take quite some time.' && echo 'Performing filesystem layer commit...' && echo ''
 
 # Run post-clone setup steps, ensuring our package lists are up to date since Setup.sh doesn't call `apt-get update`
 
@@ -124,14 +130,14 @@ RUN --mount=type=cache,target=[REDACTED_USER_HOME]/UnrealEngine/.git/ue4-gitdeps,uid=1000,g
 
 FROM source as builder
 
-# Remove the .git directory to disable UBT `git status` calls and speed up the build process
-RUN rm -rf [REDACTED_USER_HOME]/UnrealEngine/.git && echo '' && echo 'RUN directive complete. Docker will now commit the filesystem layer to disk.' && echo 'Note that for large filesystem layers this can take quite some time.' && echo 'Performing filesystem layer commit...' && echo ''
-
 # Set the changelist number in Build.version to ensure our Build ID is generated correctly
 ARG CHANGELIST
 COPY set-changelist.py /tmp/set-changelist.py
 RUN python3 /tmp/set-changelist.py [REDACTED_USER_HOME]/UnrealEngine/Engine/Build/Build.version $CHANGELIST && echo '' && echo 'RUN directive complete. Docker will now commit the filesystem layer to disk.' && echo 'Note that for large filesystem layers this can take quite some time.' && echo 'Performing filesystem layer commit...' && echo ''
 
+# Remove the .git directory to disable UBT `git status` calls and speed up the build process
+RUN rm -rf [REDACTED_USER_HOME]/UnrealEngine/.git && echo '' && echo 'RUN directive complete. Docker will now commit the filesystem layer to disk.' && echo 'Note that for large filesystem layers this can take quite some time.' && echo 'Performing filesystem layer commit...' && echo ''
+
 # Ensure UBT is built before we create the Installed Build, since Build.sh explicitly sets the
 # target .NET framework version, whereas InstalledEngineBuild.xml just uses the system default,
 # which can result in errors when running the built UBT due to the wrong version being targeted
diff --git a/Engine/Extras/Containers/Dockerfiles/linux/dev-slim/set-changelist.py b/Engine/Extras/Containers/Dockerfiles/linux/dev-slim/set-changelist.py
--- a/Engine/Extras/Containers/Dockerfiles/linux/dev-slim/set-changelist.py
+++ b/Engine/Extras/Containers/Dockerfiles/linux/dev-slim/set-changelist.py
@@ -1,5 +1,7 @@
 #!/usr/bin/env python3
-import json, sys
+from os.path import dirname
+from subprocess import run, PIPE
+import json, re, sys
 
 
 def readFile(filename):
@@ -12,9 +14,50 @@ def writeFile(filename, data):
         f.write(data.encode("utf-8"))
 
 
-# Update the `Changelist` field to reflect the `CompatibleChangelist` field in our version file, unless a specific value was provided
+# Determine whether a changelist override value was specified
+changelistOverride = None
+if len(sys.argv) > 2:
+
+    # If the override was "auto" then attempt to retrieve the CL number from the git commit message
+    if sys.argv[2] == "auto":
+
+        # Retrieve the commit message from git
+        engineRoot = dirname(dirname(dirname(sys.argv[1])))
+        commitMessage = run(
+            ["git", "log", "-n", "1", "--format=%s%n%b"],
+            cwd=engineRoot,
+            stdout=PIPE,
+            stderr=PIPE,
+            universal_newlines=True,
+        ).stdout.strip()
+
+        # If the commit is a tagged engine release then it won't have a CL number, and using "auto" is user error
+        if re.fullmatch("[0-9\\.]+ release", commitMessage) is not None:
+            print(
+                "Error: you are attempting to automatically retrieve the CL number for a tagged Unreal Engine release.\n"
+                "For hotfix releases of the Unreal Engine, a CL override is not required and should not be specified.\n"
+                "For supported .0 releases of the Unreal Engine, ue4-docker ships with known CL numbers, so an override should not be necessary.",
+                file=sys.stderr,
+            )
+            sys.exit(1)
+
+        # Attempt to extract the CL number from the commit message
+        match = re.search("\\[CL ([0-9]+) by .+ in .+ branch\\]", commitMessage)
+        if match is not None:
+            changelistOverride = int(match.group(1))
+        else:
+            print(
+                "Error: failed to find a CL number in the git commit message! This was the commit message:\n\n"
+                + commitMessage,
+                file=sys.stderr,
+            )
+            sys.exit(1)
+
+    else:
+        changelistOverride = int(sys.argv[2])
+
+# Update the `Changelist` field to reflect the override if it was supplied, or else the `CompatibleChangelist` field in our version file
 versionFile = sys.argv[1]
-changelistOverride = int(sys.argv[2]) if len(sys.argv) > 2 else None
 details = json.loads(readFile(versionFile))
 details["Changelist"] = (
     changelistOverride
diff --git a/Engine/Extras/Containers/Dockerfiles/linux/dev/Dockerfile b/Engine/Extras/Containers/Dockerfiles/linux/dev/Dockerfile
--- a/Engine/Extras/Containers/Dockerfiles/linux/dev/Dockerfile
+++ b/Engine/Extras/Containers/Dockerfiles/linux/dev/Dockerfile
@@ -1,4 +1,4 @@
-# This file was generated by ue4-docker version 0.0.93 with the following options:
+# This file was generated by ue4-docker version 0.0.94 with the following options:
 # 
 # - buildgraph_args: "-set:WithLinuxArm64=false -set:WithClient=true -set:WithServer=true"
 # - combine: true
@@ -98,7 +98,7 @@ FROM prerequisites as source
 # The git repository that we will clone
 ARG GIT_REPO=""
 
-# The git branch/tag that we will checkout
+# The git branch/tag/commit that we will checkout
 ARG GIT_BRANCH=""
 
 # Install our git credential helper that retrieves credentials from build secrets
@@ -112,7 +112,13 @@ RUN chmod +x /tmp/git-credential-helper-secrets.sh && echo '' && echo 'RUN direc
 ARG CHANGELIST
 RUN --mount=type=secret,id=username,uid=1000,required \
 	--mount=type=secret,id=password,uid=1000,required \
-	CHANGELIST="$CHANGELIST" git clone --progress --depth=1 -b $GIT_BRANCH $GIT_REPO [REDACTED_USER_HOME]/UnrealEngine && echo '' && echo 'RUN directive complete. Docker will now commit the filesystem layer to disk.' && echo 'Note that for large filesystem layers this can take quite some time.' && echo 'Performing filesystem layer commit...' && echo ''
+	CHANGELIST="$CHANGELIST" \
+	mkdir [REDACTED_USER_HOME]/UnrealEngine && \
+	cd [REDACTED_USER_HOME]/UnrealEngine && \
+	git init && \
+	git remote add origin "$GIT_REPO" && \
+	git fetch --progress --depth 1 origin "$GIT_BRANCH" && \
+	git checkout FETCH_HEAD && echo '' && echo 'RUN directive complete. Docker will now commit the filesystem layer to disk.' && echo 'Note that for large filesystem layers this can take quite some time.' && echo 'Performing filesystem layer commit...' && echo ''
 
 # Run post-clone setup steps, ensuring our package lists are up to date since Setup.sh doesn't call `apt-get update`
 
@@ -124,14 +130,14 @@ RUN --mount=type=cache,target=[REDACTED_USER_HOME]/UnrealEngine/.git/ue4-gitdeps,uid=1000,g
 
 FROM source as builder
 
-# Remove the .git directory to disable UBT `git status` calls and speed up the build process
-RUN rm -rf [REDACTED_USER_HOME]/UnrealEngine/.git && echo '' && echo 'RUN directive complete. Docker will now commit the filesystem layer to disk.' && echo 'Note that for large filesystem layers this can take quite some time.' && echo 'Performing filesystem layer commit...' && echo ''
-
 # Set the changelist number in Build.version to ensure our Build ID is generated correctly
 ARG CHANGELIST
 COPY set-changelist.py /tmp/set-changelist.py
 RUN python3 /tmp/set-changelist.py [REDACTED_USER_HOME]/UnrealEngine/Engine/Build/Build.version $CHANGELIST && echo '' && echo 'RUN directive complete. Docker will now commit the filesystem layer to disk.' && echo 'Note that for large filesystem layers this can take quite some time.' && echo 'Performing filesystem layer commit...' && echo ''
 
+# Remove the .git directory to disable UBT `git status` calls and speed up the build process
+RUN rm -rf [REDACTED_USER_HOME]/UnrealEngine/.git && echo '' && echo 'RUN directive complete. Docker will now commit the filesystem layer to disk.' && echo 'Note that for large filesystem layers this can take quite some time.' && echo 'Performing filesystem layer commit...' && echo ''
+
 # Ensure UBT is built before we create the Installed Build, since Build.sh explicitly sets the
 # target .NET framework version, whereas InstalledEngineBuild.xml just uses the system default,
 # which can result in errors when running the built UBT due to the wrong version being targeted
diff --git a/Engine/Extras/Containers/Dockerfiles/linux/dev/set-changelist.py b/Engine/Extras/Containers/Dockerfiles/linux/dev/set-changelist.py
--- a/Engine/Extras/Containers/Dockerfiles/linux/dev/set-changelist.py
+++ b/Engine/Extras/Containers/Dockerfiles/linux/dev/set-changelist.py
@@ -1,5 +1,7 @@
 #!/usr/bin/env python3
-import json, sys
+from os.path import dirname
+from subprocess import run, PIPE
+import json, re, sys
 
 
 def readFile(filename):
@@ -12,9 +14,50 @@ def writeFile(filename, data):
         f.write(data.encode("utf-8"))
 
 
-# Update the `Changelist` field to reflect the `CompatibleChangelist` field in our version file, unless a specific value was provided
+# Determine whether a changelist override value was specified
+changelistOverride = None
+if len(sys.argv) > 2:
+
+    # If the override was "auto" then attempt to retrieve the CL number from the git commit message
+    if sys.argv[2] == "auto":
+
+        # Retrieve the commit message from git
+        engineRoot = dirname(dirname(dirname(sys.argv[1])))
+        commitMessage = run(
+            ["git", "log", "-n", "1", "--format=%s%n%b"],
+            cwd=engineRoot,
+            stdout=PIPE,
+            stderr=PIPE,
+            universal_newlines=True,
+        ).stdout.strip()
+
+        # If the commit is a tagged engine release then it won't have a CL number, and using "auto" is user error
+        if re.fullmatch("[0-9\\.]+ release", commitMessage) is not None:
+            print(
+                "Error: you are attempting to automatically retrieve the CL number for a tagged Unreal Engine release.\n"
+                "For hotfix releases of the Unreal Engine, a CL override is not required and should not be specified.\n"
+                "For supported .0 releases of the Unreal Engine, ue4-docker ships with known CL numbers, so an override should not be necessary.",
+                file=sys.stderr,
+            )
+            sys.exit(1)
+
+        # Attempt to extract the CL number from the commit message
+        match = re.search("\\[CL ([0-9]+) by .+ in .+ branch\\]", commitMessage)
+        if match is not None:
+            changelistOverride = int(match.group(1))
+        else:
+            print(
+                "Error: failed to find a CL number in the git commit message! This was the commit message:\n\n"
+                + commitMessage,
+                file=sys.stderr,
+            )
+            sys.exit(1)
+
+    else:
+        changelistOverride = int(sys.argv[2])
+
+# Update the `Changelist` field to reflect the override if it was supplied, or else the `CompatibleChangelist` field in our version file
 versionFile = sys.argv[1]
-changelistOverride = int(sys.argv[2]) if len(sys.argv) > 2 else None
 details = json.loads(readFile(versionFile))
 details["Changelist"] = (
     changelistOverride
