Fix curl error in Windows runtime container image on Windows Server 2022

diff --git a/Engine/Extras/Containers/Dockerfiles/windows/runtime/Dockerfile b/Engine/Extras/Containers/Dockerfiles/windows/runtime/Dockerfile
--- a/Engine/Extras/Containers/Dockerfiles/windows/runtime/Dockerfile
+++ b/Engine/Extras/Containers/Dockerfiles/windows/runtime/Dockerfile
@@ -20,7 +20,7 @@ RUN xcopy /y C:\Windows\System32\avicap32.dll C:\GatheredDlls\ && `
 	xcopy /y C:\Windows\System32\ResourcePolicyClient.dll C:\GatheredDlls\
 
 # Retrieve the DirectX runtime files required by the Unreal Engine, since even the full Windows base image does not include them
-RUN curl --progress -L "https://download.microsoft.com/download/8/4/A/84A35BF1-DAFE-4AE8-82AF-AD2AE20B6B14/directx_Jun2010_redist.exe" --output %TEMP%\directx_redist.exe && `
+RUN curl --progress-bar -L "https://download.microsoft.com/download/8/4/A/84A35BF1-DAFE-4AE8-82AF-AD2AE20B6B14/directx_Jun2010_redist.exe" --output %TEMP%\directx_redist.exe && `
 	start /wait %TEMP%\directx_redist.exe /Q /T:%TEMP%\DirectX && `
 	expand %TEMP%\DirectX\APR2007_xinput_x64.cab -F:xinput1_3.dll C:\GatheredDlls\ && `
 	expand %TEMP%\DirectX\Feb2010_X3DAudio_x64.cab -F:X3DAudio1_7.dll C:\GatheredDlls\ && `
@@ -29,7 +29,7 @@ RUN curl --progress -L "https://download.microsoft.com/download/8/4/A/84A35BF1-D
 	expand %TEMP%\DirectX\Jun2010_XAudio_x64.cab -F:XAPOFX1_5.dll C:\GatheredDlls\
 
 # Retrieve the DirectX shader compiler files needed for DirectX Raytracing (DXR)
-RUN curl --progress -L "[REDACTED_PROFILE_URL]/DirectXShaderCompiler/releases/download/v1.6.2104/dxc_2021_04-20.zip" --output %TEMP%\dxc.zip && `
+RUN curl --progress-bar -L "[REDACTED_PROFILE_URL]/DirectXShaderCompiler/releases/download/v1.6.2104/dxc_2021_04-20.zip" --output %TEMP%\dxc.zip && `
 	powershell -Command "Expand-Archive -Path \"$env:TEMP\dxc.zip\" -DestinationPath $env:TEMP" && `
 	xcopy /y %TEMP%\bin\x64\dxcompiler.dll C:\GatheredDlls\ && `
 	xcopy /y %TEMP%\bin\x64\dxil.dll C:\GatheredDlls\
