Fix variations between RES and other streams due to merge issues
#rnx

diff --git a/Engine/Extras/Containers/Dockerfiles/windows/build.ps1 b/Engine/Extras/Containers/Dockerfiles/windows/build.ps1
--- a/Engine/Extras/Containers/Dockerfiles/windows/build.ps1
+++ b/Engine/Extras/Containers/Dockerfiles/windows/build.ps1
@@ -23,6 +23,12 @@ else
 		$windowsRelease = (Get-ItemProperty -Path 'Registry::HKEY_LOCAL_MACHINE\SOFTWARE\Microsoft\Windows NT\CurrentVersion' -Name ReleaseId).ReleaseId
 	}
 	
+	# Use LTSC2022 images under Windows 11 and Windows Server 2022
+	[int]$kernelBuild = (Get-ItemProperty -Path 'Registry::HKEY_LOCAL_MACHINE\SOFTWARE\Microsoft\Windows NT\CurrentVersion' -Name CurrentBuildNumber).CurrentBuildNumber
+	if ($kernelBuild -ge 20348) {
+		$windowsRelease = "ltsc2022"
+	}
+
 	# Use process isolation mode for improved performance
 	$isolation = "process"
 	
@@ -30,6 +36,12 @@ else
 	$tag = "runtime-windows"
 }
 
+# Determine the appropriate image to copy DLL files from
+$dllImage = "mcr.microsoft.com/windows"
+if ($windowsRelease -eq "ltsc2022") {
+	$dllImage = "mcr.microsoft.com/windows/server"
+}
+
 # Build our runtime container image using the correct base image for the selected Windows version
-"Building runtime container image for Windows version $windowsRelease with ``$isolation`` isolation mode..."
-docker build -t "ghcr.io/epicgames/unreal-engine:$tag" --isolation="$isolation" --build-arg "BASETAG=$windowsRelease" ./runtime
+"Building runtime container image for Windows version $windowsRelease with ``$isolation`` isolation mode and DLL files from ``$dllImage``..."
+docker build -t "ghcr.io/epicgames/unreal-engine:$tag" --isolation="$isolation" --build-arg "DLL_IMAGE=$dllImage" --build-arg "BASETAG=$windowsRelease" ./runtime
diff --git a/Engine/Extras/Containers/Dockerfiles/windows/runtime/Dockerfile b/Engine/Extras/Containers/Dockerfiles/windows/runtime/Dockerfile
--- a/Engine/Extras/Containers/Dockerfiles/windows/runtime/Dockerfile
+++ b/Engine/Extras/Containers/Dockerfiles/windows/runtime/Dockerfile
@@ -1,6 +1,7 @@
 # escape=`
+ARG DLL_IMAGE
 ARG BASETAG
-FROM mcr.microsoft.com/windows:${BASETAG} AS full
+FROM ${DLL_IMAGE}:${BASETAG} AS full
 
 # Gather the system DLLs that we need from the full Windows base image
 RUN xcopy /y C:\Windows\System32\avicap32.dll C:\GatheredDlls\ && `
diff --git a/Engine/Plugins/FX/Niagara/Source/NiagaraEditor/Private/Widgets/SNiagaraPinTypeSelector.cpp b/Engine/Plugins/FX/Niagara/Source/NiagaraEditor/Private/Widgets/SNiagaraPinTypeSelector.cpp
--- a/Engine/Plugins/FX/Niagara/Source/NiagaraEditor/Private/Widgets/SNiagaraPinTypeSelector.cpp
+++ b/Engine/Plugins/FX/Niagara/Source/NiagaraEditor/Private/Widgets/SNiagaraPinTypeSelector.cpp
@@ -4,7 +4,6 @@
 
 #include "NiagaraEditorStyle.h"
 #include "Widgets/Images/SImage.h"
-//#include "Widgets/SNiagaraParameterPanel.h"
 #include "Widgets/SNiagaraParameterMenu.h"
 #include "NiagaraNode.h"
 #include "NiagaraEditorUtilities.h"
diff --git a/Engine/Source/Editor/GameProjectGeneration/Classes/TemplateProjectDefs.h b/Engine/Source/Editor/GameProjectGeneration/Classes/TemplateProjectDefs.h
--- a/Engine/Source/Editor/GameProjectGeneration/Classes/TemplateProjectDefs.h
+++ b/Engine/Source/Editor/GameProjectGeneration/Classes/TemplateProjectDefs.h
@@ -203,7 +203,7 @@ struct FTemplateCategoryDef
 
 	/** Is this a major top-level category? Major categories are displayed as full rows, eg. the Game category.*/
 	UPROPERTY()
-		bool IsMajorCategory{ false };
+	bool IsMajorCategory=false;
 };
 
 UCLASS(config=TemplateCategories)
