Re-generate Dockerfiles to fix build errors

diff --git a/Engine/Extras/Containers/Dockerfiles/linux/dev-slim/Dockerfile b/Engine/Extras/Containers/Dockerfiles/linux/dev-slim/Dockerfile
--- a/Engine/Extras/Containers/Dockerfiles/linux/dev-slim/Dockerfile
+++ b/Engine/Extras/Containers/Dockerfiles/linux/dev-slim/Dockerfile
@@ -1,6 +1,6 @@
-# This file was generated by ue4-docker version 0.0.94 with the following options:
+# This file was generated by ue4-docker version 0.0.98 with the following options:
 # 
-# - buildgraph_args: "-set:WithLinuxArm64=false -set:WithClient=true -set:WithServer=true"
+# - buildgraph_args: "-set:HostPlatformOnly=true -set:WithLinuxArm64=false -set:WithClient=true -set:WithServer=true"
 # - combine: true
 # - credential_mode: "secrets"
 # - disable_all_patches: true
@@ -60,6 +60,7 @@ RUN apt-get update && apt-get install -y --no-install-recommends \
 			libsm6 \
 			libxcomposite1 \
 			libxcursor1 \
+			libxdamage1 \
 			libxi6 \
 			libxrandr2 \
 			libxrender1 \
@@ -138,14 +139,12 @@ RUN python3 /tmp/set-changelist.py [REDACTED_USER_HOME]/UnrealEngine/Engine/Build/Build.ver
 # Remove the .git directory to disable UBT `git status` calls and speed up the build process
 RUN rm -rf [REDACTED_USER_HOME]/UnrealEngine/.git && echo '' && echo 'RUN directive complete. Docker will now commit the filesystem layer to disk.' && echo 'Note that for large filesystem layers this can take quite some time.' && echo 'Performing filesystem layer commit...' && echo ''
 
-# Ensure UBT is built before we create the Installed Build, since Build.sh explicitly sets the
-# target .NET framework version, whereas InstalledEngineBuild.xml just uses the system default,
-# which can result in errors when running the built UBT due to the wrong version being targeted
-RUN ./Engine/Build/BatchFiles/Linux/Build.sh UnrealHeaderTool Linux Development -SkipBuild && echo '' && echo 'RUN directive complete. Docker will now commit the filesystem layer to disk.' && echo 'Note that for large filesystem layers this can take quite some time.' && echo 'Performing filesystem layer commit...' && echo ''
-
 # Create an Installed Build of the Engine
 WORKDIR [REDACTED_USER_HOME]/UnrealEngine
-RUN ./Engine/Build/BatchFiles/RunUAT.sh BuildGraph -target="Make Installed Build Linux" -script=Engine/Build/InstalledEngineBuild.xml -set:HostPlatformOnly=true -set:WithDDC=true -set:WithLinuxArm64=false -set:WithClient=true -set:WithServer=true && \
+RUN ./Engine/Build/BatchFiles/RunUAT.sh BuildGraph \
+    -target="Make Installed Build Linux" \
+    -script=Engine/Build/InstalledEngineBuild.xml \
+    -set:WithDDC=true -set:HostPlatformOnly=true -set:WithLinuxArm64=false -set:WithClient=true -set:WithServer=true && \
 	rm -R -f [REDACTED_USER_HOME]/UnrealEngine/LocalBuilds/InstalledDDC && echo '' && echo 'RUN directive complete. Docker will now commit the filesystem layer to disk.' && echo 'Note that for large filesystem layers this can take quite some time.' && echo 'Performing filesystem layer commit...' && echo ''
 
 # Split out components (DDC, debug symbols, template projects) so they can be copied into the final container image as separate filesystem layers
diff --git a/Engine/Extras/Containers/Dockerfiles/linux/dev/Dockerfile b/Engine/Extras/Containers/Dockerfiles/linux/dev/Dockerfile
--- a/Engine/Extras/Containers/Dockerfiles/linux/dev/Dockerfile
+++ b/Engine/Extras/Containers/Dockerfiles/linux/dev/Dockerfile
@@ -1,6 +1,6 @@
-# This file was generated by ue4-docker version 0.0.94 with the following options:
+# This file was generated by ue4-docker version 0.0.98 with the following options:
 # 
-# - buildgraph_args: "-set:WithLinuxArm64=false -set:WithClient=true -set:WithServer=true"
+# - buildgraph_args: "-set:HostPlatformOnly=true -set:WithLinuxArm64=false -set:WithClient=true -set:WithServer=true"
 # - combine: true
 # - credential_mode: "secrets"
 # - disable_all_patches: true
@@ -60,6 +60,7 @@ RUN apt-get update && apt-get install -y --no-install-recommends \
 			libsm6 \
 			libxcomposite1 \
 			libxcursor1 \
+			libxdamage1 \
 			libxi6 \
 			libxrandr2 \
 			libxrender1 \
@@ -138,14 +139,12 @@ RUN python3 /tmp/set-changelist.py [REDACTED_USER_HOME]/UnrealEngine/Engine/Build/Build.ver
 # Remove the .git directory to disable UBT `git status` calls and speed up the build process
 RUN rm -rf [REDACTED_USER_HOME]/UnrealEngine/.git && echo '' && echo 'RUN directive complete. Docker will now commit the filesystem layer to disk.' && echo 'Note that for large filesystem layers this can take quite some time.' && echo 'Performing filesystem layer commit...' && echo ''
 
-# Ensure UBT is built before we create the Installed Build, since Build.sh explicitly sets the
-# target .NET framework version, whereas InstalledEngineBuild.xml just uses the system default,
-# which can result in errors when running the built UBT due to the wrong version being targeted
-RUN ./Engine/Build/BatchFiles/Linux/Build.sh UnrealHeaderTool Linux Development -SkipBuild && echo '' && echo 'RUN directive complete. Docker will now commit the filesystem layer to disk.' && echo 'Note that for large filesystem layers this can take quite some time.' && echo 'Performing filesystem layer commit...' && echo ''
-
 # Create an Installed Build of the Engine
 WORKDIR [REDACTED_USER_HOME]/UnrealEngine
-RUN ./Engine/Build/BatchFiles/RunUAT.sh BuildGraph -target="Make Installed Build Linux" -script=Engine/Build/InstalledEngineBuild.xml -set:HostPlatformOnly=true -set:WithDDC=true -set:WithLinuxArm64=false -set:WithClient=true -set:WithServer=true && \
+RUN ./Engine/Build/BatchFiles/RunUAT.sh BuildGraph \
+    -target="Make Installed Build Linux" \
+    -script=Engine/Build/InstalledEngineBuild.xml \
+    -set:WithDDC=true -set:HostPlatformOnly=true -set:WithLinuxArm64=false -set:WithClient=true -set:WithServer=true && \
 	rm -R -f [REDACTED_USER_HOME]/UnrealEngine/LocalBuilds/InstalledDDC && echo '' && echo 'RUN directive complete. Docker will now commit the filesystem layer to disk.' && echo 'Note that for large filesystem layers this can take quite some time.' && echo 'Performing filesystem layer commit...' && echo ''
 
 # Split out components (DDC, debug symbols, template projects) so they can be copied into the final container image as separate filesystem layers
diff --git a/Engine/Extras/Containers/Dockerfiles/linux/generate.sh b/Engine/Extras/Containers/Dockerfiles/linux/generate.sh
--- a/Engine/Extras/Containers/Dockerfiles/linux/generate.sh
+++ b/Engine/Extras/Containers/Dockerfiles/linux/generate.sh
@@ -31,7 +31,7 @@ function generateDockerfile {
 		--no-engine --no-full
 		
 		# This disables building the Engine for AArch64 when creating an Installed Build, and enables support for client and dedicated server targets
-		--opt buildgraph-args='-set:WithLinuxArm64=false -set:WithClient=true -set:WithServer=true'
+		--opt buildgraph-args='-set:HostPlatformOnly=true -set:WithLinuxArm64=false -set:WithClient=true -set:WithServer=true'
 		
 		# This enables the use of BuildKit build secrets, which is necessary in order to build images independently of ue4-docker itself
 		--opt credential-mode=secrets
diff --git a/Engine/Extras/Containers/Dockerfiles/windows/runtime/Dockerfile b/Engine/Extras/Containers/Dockerfiles/windows/runtime/Dockerfile
--- a/Engine/Extras/Containers/Dockerfiles/windows/runtime/Dockerfile
+++ b/Engine/Extras/Containers/Dockerfiles/windows/runtime/Dockerfile
@@ -5,19 +5,24 @@ FROM ${DLL_IMAGE}:${BASETAG} AS full
 
 # Gather the system DLLs that we need from the full Windows base image
 RUN xcopy /y C:\Windows\System32\avicap32.dll C:\GatheredDlls\ && `
+	xcopy /y C:\Windows\System32\avifil32.dll C:\GatheredDlls\ && `
 	xcopy /y C:\Windows\System32\avrt.dll C:\GatheredDlls\ && `
 	xcopy /y C:\Windows\System32\d3d10warp.dll C:\GatheredDlls\ && `
 	xcopy /y C:\Windows\System32\D3DSCache.dll C:\GatheredDlls\ && `
 	xcopy /y C:\Windows\System32\dsound.dll C:\GatheredDlls\ && `
 	xcopy /y C:\Windows\System32\dxva2.dll C:\GatheredDlls\ && `
 	xcopy /y C:\Windows\System32\glu32.dll C:\GatheredDlls\ && `
+	xcopy /y C:\Windows\System32\ksuser.dll C:\GatheredDlls\ && `
 	xcopy /y C:\Windows\System32\mf.dll C:\GatheredDlls\ && `
+	xcopy /y C:\Windows\System32\mfcore.dll C:\GatheredDlls\ && `
 	xcopy /y C:\Windows\System32\mfplat.dll C:\GatheredDlls\ && `
 	xcopy /y C:\Windows\System32\mfplay.dll C:\GatheredDlls\ && `
 	xcopy /y C:\Windows\System32\mfreadwrite.dll C:\GatheredDlls\ && `
+	xcopy /y C:\Windows\System32\msacm32.dll C:\GatheredDlls\ && `
 	xcopy /y C:\Windows\System32\msdmo.dll C:\GatheredDlls\ && `
 	xcopy /y C:\Windows\System32\msvfw32.dll C:\GatheredDlls\ && `
 	xcopy /y C:\Windows\System32\opengl32.dll C:\GatheredDlls\ && `
+	xcopy /y C:\Windows\System32\ResampleDMO.dll C:\GatheredDlls\ && `
 	xcopy /y C:\Windows\System32\ResourcePolicyClient.dll C:\GatheredDlls\
 
 # Retrieve the DirectX runtime files required by the Unreal Engine, since even the full Windows base image does not include them
