Update Dockerfiles and fix Docker Compose Pixel Streaming demo

diff --git a/Engine/Extras/Containers/Dockerfiles/linux/build.sh b/Engine/Extras/Containers/Dockerfiles/linux/build.sh
--- a/Engine/Extras/Containers/Dockerfiles/linux/build.sh
+++ b/Engine/Extras/Containers/Dockerfiles/linux/build.sh
@@ -57,6 +57,7 @@ args=(
 	--build-arg "CHANGELIST=${CHANGELIST_OVERRIDE}"
 	--secret id=username,src=username.txt
 	--secret id=password,src=password.txt
+	--platform linux/amd64
 	--progress=plain
 )
 
@@ -91,12 +92,6 @@ docker buildx build \
 	"${args[@]}" \
 	./dev-slim
 
-# Build the container image for the Pixel Streaming signalling server, which pulls files from our other images
-docker build \
-	-t "ghcr.io/epicgames/pixel-streaming-signalling-server:${UNREAL_ENGINE_RELEASE}" \
-	--build-arg "UNREAL_ENGINE_RELEASE=${UNREAL_ENGINE_RELEASE}" \
-	./signalling-server
-
 # Build the container image for the Multi-User Editing server, which pulls files from our other images
 docker build \
 	-t "ghcr.io/epicgames/multi-user-server:${UNREAL_ENGINE_RELEASE}" \
diff --git a/Engine/Extras/Containers/Dockerfiles/linux/dev-slim/Dockerfile b/Engine/Extras/Containers/Dockerfiles/linux/dev-slim/Dockerfile
--- a/Engine/Extras/Containers/Dockerfiles/linux/dev-slim/Dockerfile
+++ b/Engine/Extras/Containers/Dockerfiles/linux/dev-slim/Dockerfile
@@ -1,6 +1,6 @@
-# This file was generated by ue4-docker version 0.0.98 with the following options:
+# This file was generated by ue4-docker version 0.0.102 with the following options:
 # 
-# - buildgraph_args: "-set:HostPlatformOnly=true -set:WithLinuxArm64=false -set:WithClient=true -set:WithServer=true"
+# - buildgraph_args: "-set:WithLinuxArm64=false -set:WithClient=true -set:WithServer=true"
 # - combine: true
 # - credential_mode: "secrets"
 # - disable_all_patches: true
@@ -20,10 +20,18 @@ FROM ${BASEIMAGE} as prerequisites
 # Disable interactive prompts during package installation
 ENV DEBIAN_FRONTEND=noninteractive
 
+# Enable CUDA support for NVIDIA GPUs (even when not using a CUDA base image), since evidently some versions of UE unconditionally assume
+# `libcuda.so.1` exists when using the NVIDIA proprietary drivers, and will fail to initialise the Vulkan RHI if it is missing
+ENV NVIDIA_DRIVER_CAPABILITIES ${NVIDIA_DRIVER_CAPABILITIES},compute
+
 # Add the "display" driver capability for NVIDIA GPUs
-# (This also allows us to run the Editor from an interactive container by bind-mounting the host system's X11 socket)
+# (This allows us to run the Editor from an interactive container by bind-mounting the host system's X11 socket)
 ENV NVIDIA_DRIVER_CAPABILITIES ${NVIDIA_DRIVER_CAPABILITIES},display
 
+# Enable NVENC support for use by Unreal Engine plugins that depend on it (e.g. Pixel Streaming)
+# (Note that adding `video` seems to implicitly enable `compute` as well, but we include separate directives here to clearly indicate the purpose of both)
+ENV NVIDIA_DRIVER_CAPABILITIES ${NVIDIA_DRIVER_CAPABILITIES},video
+
 # Install our build prerequisites
 RUN apt-get update && apt-get install -y --no-install-recommends \
 		build-essential \
@@ -62,6 +70,7 @@ RUN apt-get update && apt-get install -y --no-install-recommends \
 			libxcursor1 \
 			libxdamage1 \
 			libxi6 \
+			libxkbcommon-x11-0 \
 			libxrandr2 \
 			libxrender1 \
 			libxss1 \
@@ -123,10 +132,19 @@ RUN --mount=type=secret,id=username,uid=1000,required \
 
 # Run post-clone setup steps, ensuring our package lists are up to date since Setup.sh doesn't call `apt-get update`
 
-# When running with BuildKit, we use a cache mount to cache the dependency data in `.git/ue4-gitdeps` across multiple build invocations
+# Ensure Setup.sh uses the same cache path when building either UE4 or UE5
+ENV UE_GITDEPS=[REDACTED_USER_HOME]/gitdeps
+ENV UE4_GITDEPS=[REDACTED_USER_HOME]/gitdeps
+RUN mkdir "$UE_GITDEPS" && echo '' && echo 'RUN directive complete. Docker will now commit the filesystem layer to disk.' && echo 'Note that for large filesystem layers this can take quite some time.' && echo 'Performing filesystem layer commit...' && echo ''
+
+# When running with BuildKit, we use a cache mount to cache the dependency data across multiple build invocations
 WORKDIR [REDACTED_USER_HOME]/UnrealEngine
-RUN --mount=type=cache,target=[REDACTED_USER_HOME]/UnrealEngine/.git/ue4-gitdeps,uid=1000,gid=1000 sudo apt-get update && \
-	./Setup.sh && \
+RUN --mount=type=cache,target=[REDACTED_USER_HOME]/gitdeps,uid=1000,gid=1000 sudo apt-get update && \
+	./Setup.sh \
+		--exclude=Android \
+		--exclude=Mac \
+		--exclude=Win32 \
+		--exclude=Win64 && \
 	sudo rm -rf /var/lib/apt/lists/* && echo '' && echo 'RUN directive complete. Docker will now commit the filesystem layer to disk.' && echo 'Note that for large filesystem layers this can take quite some time.' && echo 'Performing filesystem layer commit...' && echo ''
 
 FROM source as builder
@@ -139,12 +157,19 @@ RUN python3 /tmp/set-changelist.py [REDACTED_USER_HOME]/UnrealEngine/Engine/Build/Build.ver
 # Remove the .git directory to disable UBT `git status` calls and speed up the build process
 RUN rm -rf [REDACTED_USER_HOME]/UnrealEngine/.git && echo '' && echo 'RUN directive complete. Docker will now commit the filesystem layer to disk.' && echo 'Note that for large filesystem layers this can take quite some time.' && echo 'Performing filesystem layer commit...' && echo ''
 
+# Ensure UBT is built before we create the Installed Build, since Build.sh explicitly sets the
+# target .NET Framework version, whereas InstalledEngineBuild.xml just uses the system default,
+# which can result in errors when running the built UBT due to the wrong version being targeted
+RUN ./Engine/Build/BatchFiles/Linux/Build.sh UnrealHeaderTool Linux Development -SkipBuild -buildubt && echo '' && echo 'RUN directive complete. Docker will now commit the filesystem layer to disk.' && echo 'Note that for large filesystem layers this can take quite some time.' && echo 'Performing filesystem layer commit...' && echo ''
+
 # Create an Installed Build of the Engine
 WORKDIR [REDACTED_USER_HOME]/UnrealEngine
 RUN ./Engine/Build/BatchFiles/RunUAT.sh BuildGraph \
     -target="Make Installed Build Linux" \
     -script=Engine/Build/InstalledEngineBuild.xml \
-    -set:WithDDC=true -set:HostPlatformOnly=true -set:WithLinuxArm64=false -set:WithClient=true -set:WithServer=true && \
+    -set:HostPlatformOnly=true \
+    -set:WithDDC=true \
+    -set:WithLinuxArm64=false -set:WithClient=true -set:WithServer=true && \
 	rm -R -f [REDACTED_USER_HOME]/UnrealEngine/LocalBuilds/InstalledDDC && echo '' && echo 'RUN directive complete. Docker will now commit the filesystem layer to disk.' && echo 'Note that for large filesystem layers this can take quite some time.' && echo 'Performing filesystem layer commit...' && echo ''
 
 # Split out components (DDC, debug symbols, template projects) so they can be copied into the final container image as separate filesystem layers
diff --git a/Engine/Extras/Containers/Dockerfiles/linux/dev/Dockerfile b/Engine/Extras/Containers/Dockerfiles/linux/dev/Dockerfile
--- a/Engine/Extras/Containers/Dockerfiles/linux/dev/Dockerfile
+++ b/Engine/Extras/Containers/Dockerfiles/linux/dev/Dockerfile
@@ -1,6 +1,6 @@
-# This file was generated by ue4-docker version 0.0.98 with the following options:
+# This file was generated by ue4-docker version 0.0.102 with the following options:
 # 
-# - buildgraph_args: "-set:HostPlatformOnly=true -set:WithLinuxArm64=false -set:WithClient=true -set:WithServer=true"
+# - buildgraph_args: "-set:WithLinuxArm64=false -set:WithClient=true -set:WithServer=true"
 # - combine: true
 # - credential_mode: "secrets"
 # - disable_all_patches: true
@@ -20,10 +20,18 @@ FROM ${BASEIMAGE} as prerequisites
 # Disable interactive prompts during package installation
 ENV DEBIAN_FRONTEND=noninteractive
 
+# Enable CUDA support for NVIDIA GPUs (even when not using a CUDA base image), since evidently some versions of UE unconditionally assume
+# `libcuda.so.1` exists when using the NVIDIA proprietary drivers, and will fail to initialise the Vulkan RHI if it is missing
+ENV NVIDIA_DRIVER_CAPABILITIES ${NVIDIA_DRIVER_CAPABILITIES},compute
+
 # Add the "display" driver capability for NVIDIA GPUs
-# (This also allows us to run the Editor from an interactive container by bind-mounting the host system's X11 socket)
+# (This allows us to run the Editor from an interactive container by bind-mounting the host system's X11 socket)
 ENV NVIDIA_DRIVER_CAPABILITIES ${NVIDIA_DRIVER_CAPABILITIES},display
 
+# Enable NVENC support for use by Unreal Engine plugins that depend on it (e.g. Pixel Streaming)
+# (Note that adding `video` seems to implicitly enable `compute` as well, but we include separate directives here to clearly indicate the purpose of both)
+ENV NVIDIA_DRIVER_CAPABILITIES ${NVIDIA_DRIVER_CAPABILITIES},video
+
 # Install our build prerequisites
 RUN apt-get update && apt-get install -y --no-install-recommends \
 		build-essential \
@@ -62,6 +70,7 @@ RUN apt-get update && apt-get install -y --no-install-recommends \
 			libxcursor1 \
 			libxdamage1 \
 			libxi6 \
+			libxkbcommon-x11-0 \
 			libxrandr2 \
 			libxrender1 \
 			libxss1 \
@@ -123,10 +132,19 @@ RUN --mount=type=secret,id=username,uid=1000,required \
 
 # Run post-clone setup steps, ensuring our package lists are up to date since Setup.sh doesn't call `apt-get update`
 
-# When running with BuildKit, we use a cache mount to cache the dependency data in `.git/ue4-gitdeps` across multiple build invocations
+# Ensure Setup.sh uses the same cache path when building either UE4 or UE5
+ENV UE_GITDEPS=[REDACTED_USER_HOME]/gitdeps
+ENV UE4_GITDEPS=[REDACTED_USER_HOME]/gitdeps
+RUN mkdir "$UE_GITDEPS" && echo '' && echo 'RUN directive complete. Docker will now commit the filesystem layer to disk.' && echo 'Note that for large filesystem layers this can take quite some time.' && echo 'Performing filesystem layer commit...' && echo ''
+
+# When running with BuildKit, we use a cache mount to cache the dependency data across multiple build invocations
 WORKDIR [REDACTED_USER_HOME]/UnrealEngine
-RUN --mount=type=cache,target=[REDACTED_USER_HOME]/UnrealEngine/.git/ue4-gitdeps,uid=1000,gid=1000 sudo apt-get update && \
-	./Setup.sh && \
+RUN --mount=type=cache,target=[REDACTED_USER_HOME]/gitdeps,uid=1000,gid=1000 sudo apt-get update && \
+	./Setup.sh \
+		--exclude=Android \
+		--exclude=Mac \
+		--exclude=Win32 \
+		--exclude=Win64 && \
 	sudo rm -rf /var/lib/apt/lists/* && echo '' && echo 'RUN directive complete. Docker will now commit the filesystem layer to disk.' && echo 'Note that for large filesystem layers this can take quite some time.' && echo 'Performing filesystem layer commit...' && echo ''
 
 FROM source as builder
@@ -139,12 +157,19 @@ RUN python3 /tmp/set-changelist.py [REDACTED_USER_HOME]/UnrealEngine/Engine/Build/Build.ver
 # Remove the .git directory to disable UBT `git status` calls and speed up the build process
 RUN rm -rf [REDACTED_USER_HOME]/UnrealEngine/.git && echo '' && echo 'RUN directive complete. Docker will now commit the filesystem layer to disk.' && echo 'Note that for large filesystem layers this can take quite some time.' && echo 'Performing filesystem layer commit...' && echo ''
 
+# Ensure UBT is built before we create the Installed Build, since Build.sh explicitly sets the
+# target .NET Framework version, whereas InstalledEngineBuild.xml just uses the system default,
+# which can result in errors when running the built UBT due to the wrong version being targeted
+RUN ./Engine/Build/BatchFiles/Linux/Build.sh UnrealHeaderTool Linux Development -SkipBuild -buildubt && echo '' && echo 'RUN directive complete. Docker will now commit the filesystem layer to disk.' && echo 'Note that for large filesystem layers this can take quite some time.' && echo 'Performing filesystem layer commit...' && echo ''
+
 # Create an Installed Build of the Engine
 WORKDIR [REDACTED_USER_HOME]/UnrealEngine
 RUN ./Engine/Build/BatchFiles/RunUAT.sh BuildGraph \
     -target="Make Installed Build Linux" \
     -script=Engine/Build/InstalledEngineBuild.xml \
-    -set:WithDDC=true -set:HostPlatformOnly=true -set:WithLinuxArm64=false -set:WithClient=true -set:WithServer=true && \
+    -set:HostPlatformOnly=true \
+    -set:WithDDC=true \
+    -set:WithLinuxArm64=false -set:WithClient=true -set:WithServer=true && \
 	rm -R -f [REDACTED_USER_HOME]/UnrealEngine/LocalBuilds/InstalledDDC && echo '' && echo 'RUN directive complete. Docker will now commit the filesystem layer to disk.' && echo 'Note that for large filesystem layers this can take quite some time.' && echo 'Performing filesystem layer commit...' && echo ''
 
 # Split out components (DDC, debug symbols, template projects) so they can be copied into the final container image as separate filesystem layers
diff --git a/Engine/Extras/Containers/Dockerfiles/linux/generate.sh b/Engine/Extras/Containers/Dockerfiles/linux/generate.sh
--- a/Engine/Extras/Containers/Dockerfiles/linux/generate.sh
+++ b/Engine/Extras/Containers/Dockerfiles/linux/generate.sh
@@ -27,11 +27,11 @@ function generateDockerfile {
 		# This will combine the generated Dockerfiles into a single Dockerfile that uses a multi-stage build
 		--combine
 		
-		# These flags prevent Dockerfiles from being generated for the `ue4-engine` and `ue4-full` images, which we are not interested in
-		--no-engine --no-full
+		# This specifies the individual images whose Dockerfiles will be combined
+		--target=build-prerequisites,source,minimal
 		
 		# This disables building the Engine for AArch64 when creating an Installed Build, and enables support for client and dedicated server targets
-		--opt buildgraph-args='-set:HostPlatformOnly=true -set:WithLinuxArm64=false -set:WithClient=true -set:WithServer=true'
+		--opt buildgraph-args='-set:WithLinuxArm64=false -set:WithClient=true -set:WithServer=true'
 		
 		# This enables the use of BuildKit build secrets, which is necessary in order to build images independently of ue4-docker itself
 		--opt credential-mode=secrets
diff --git a/Engine/Extras/Containers/Dockerfiles/linux/signalling-server/Dockerfile b/Engine/Extras/Containers/Dockerfiles/linux/signalling-server/Dockerfile
deleted file mode 100644
diff --git a/Engine/Extras/Containers/Examples/PixelStreaming/run.sh b/Engine/Extras/Containers/Examples/PixelStreaming/run.sh
--- a/Engine/Extras/Containers/Examples/PixelStreaming/run.sh
+++ b/Engine/Extras/Containers/Examples/PixelStreaming/run.sh
@@ -6,6 +6,9 @@ if [[ ! -z "$1" ]]; then
 	UNREAL_ENGINE_RELEASE="$1"
 fi
 
+# Extract the short version of the release number (i.e. just `major.minor` rather than `major.minor.hotfix`)
+UNREAL_ENGINE_RELEASE_SHORT=`echo "$UNREAL_ENGINE_RELEASE" | grep --extended-regexp --only-matching -e '[0-9]+\.[0-9]+'`
+
 # Determine whether we are instructing Docker Compose to rebuild the project container image even if it already exists
 COMPOSE_FLAGS=""
 if [[ "$*" == *"--rebuild"* ]]; then
@@ -29,12 +32,19 @@ else
 	exit 1
 fi
 
+# Determine whether to use the old, standalone version of Docker Compose or the new plugin version
+COMPOSE_COMMAND='docker compose'
+if [[ ! -z `which docker-compose` ]]; then
+	COMPOSE_COMMAND='docker-compose'
+fi
+
 # Retrieve the public IP address of the host system
 PUBLIC_IP=$($HTTPS_COMMAND 'https://api.ipify.org')
 
 # Run the Pixel Streaming example
 export UNREAL_ENGINE_RELEASE
+export UNREAL_ENGINE_RELEASE_SHORT
 export EXTRA_PEERCONNECTION_OPTIONS
 export PUBLIC_IP
 export PWD=$(pwd)
-docker-compose up --force-recreate $COMPOSE_FLAGS
+$COMPOSE_COMMAND up --force-recreate $COMPOSE_FLAGS
