Update Linux runtime container images to use Ubuntu 20.04

diff --git a/Engine/Extras/Containers/Dockerfiles/linux/runtime-pixel-streaming/Dockerfile b/Engine/Extras/Containers/Dockerfiles/linux/runtime-pixel-streaming/Dockerfile
--- a/Engine/Extras/Containers/Dockerfiles/linux/runtime-pixel-streaming/Dockerfile
+++ b/Engine/Extras/Containers/Dockerfiles/linux/runtime-pixel-streaming/Dockerfile
@@ -1,4 +1 @@
-FROM adamrehn/ue4-runtime:18.04-cudagl11.1.1
-
-# Enable the NVIDIA driver capabilities required by the NVENC video encoding API
-ENV NVIDIA_DRIVER_CAPABILITIES ${NVIDIA_DRIVER_CAPABILITIES},video
+FROM adamrehn/ue4-runtime:20.04-cudagl11
diff --git a/Engine/Extras/Containers/Dockerfiles/linux/runtime/Dockerfile b/Engine/Extras/Containers/Dockerfiles/linux/runtime/Dockerfile
--- a/Engine/Extras/Containers/Dockerfiles/linux/runtime/Dockerfile
+++ b/Engine/Extras/Containers/Dockerfiles/linux/runtime/Dockerfile
@@ -1 +1 @@
-FROM adamrehn/ue4-runtime:18.04-vulkan-noaudio
+FROM adamrehn/ue4-runtime:20.04-vulkan-noaudio
