Re-generate development Dockerfiles with ue4-docker version 0.0.110 and migrate to Ubuntu 22.04 base image.

This migrates the development images to a currently supported version of Ubuntu, and fixes image build errors for Unreal Engine 5.3 and newer.

diff --git a/Engine/Extras/Containers/Dockerfiles/linux/build.sh b/Engine/Extras/Containers/Dockerfiles/linux/build.sh
--- a/Engine/Extras/Containers/Dockerfiles/linux/build.sh
+++ b/Engine/Extras/Containers/Dockerfiles/linux/build.sh
@@ -51,7 +51,7 @@ fi
 
 # Assemble the common arguments to pass to the `docker buildx build` command
 args=(
-	--build-arg 'BASEIMAGE=nvidia/opengl:1.0-glvnd-devel-ubuntu18.04'
+	--build-arg 'BASEIMAGE=nvidia/opengl:1.2-glvnd-devel-ubuntu22.04'
 	--build-arg "GIT_REPO=${GIT_REPO}"
 	--build-arg "GIT_BRANCH=${GIT_BRANCH}"
 	--build-arg "CHANGELIST=${CHANGELIST_OVERRIDE}"
diff --git a/Engine/Extras/Containers/Dockerfiles/linux/dev-slim/Dockerfile b/Engine/Extras/Containers/Dockerfiles/linux/dev-slim/Dockerfile
--- a/Engine/Extras/Containers/Dockerfiles/linux/dev-slim/Dockerfile
+++ b/Engine/Extras/Containers/Dockerfiles/linux/dev-slim/Dockerfile
@@ -1,12 +1,13 @@
-# This file was generated by ue4-docker version 0.0.102 with the following options:
+# This file was generated by ue4-docker version 0.0.110 with the following options:
 # 
 # - buildgraph_args: "-set:WithLinuxArm64=false -set:WithClient=true -set:WithServer=true"
 # - combine: true
 # - credential_mode: "secrets"
 # - disable_all_patches: true
 # - disable_labels: true
-# - enable_dso_patch: true
+# - enable_dso_patch: false
 # - excluded_components: {"ddc": false, "debug": true, "templates": true}
+# - gitdependencies_args: "--exclude=Android --exclude=Mac --exclude=Win32 --exclude=Win64"
 # 
 # This Dockerfile combines the steps for the following images:
 # 
@@ -82,15 +83,6 @@ RUN apt-get update && apt-get install -y --no-install-recommends \
 			xkb-data && \
 	rm -rf /var/lib/apt/lists/* && echo '' && echo 'RUN directive complete. Docker will now commit the filesystem layer to disk.' && echo 'Note that for large filesystem layers this can take quite some time.' && echo 'Performing filesystem layer commit...' && echo ''
 
-# Install the glibc DSO patch to improve Editor startup times
-RUN add-apt-repository -y ppa:slonopotamus/glibc-dso && \
-	apt-get update && \
-	apt upgrade -y libc6 && \
-	rm -rf /var/lib/apt/lists/* && echo '' && echo 'RUN directive complete. Docker will now commit the filesystem layer to disk.' && echo 'Note that for large filesystem layers this can take quite some time.' && echo 'Performing filesystem layer commit...' && echo ''
-
-# Enable the glibc DSO patch
-ENV GLIBC_TUNABLES=glibc.rtld.dynamic_sort=2
-
 # Disable the default "lecture" message the first time a user runs a command using sudo
 RUN echo 'Defaults lecture="never"' >> /etc/sudoers && echo '' && echo 'RUN directive complete. Docker will now commit the filesystem layer to disk.' && echo 'Note that for large filesystem layers this can take quite some time.' && echo 'Performing filesystem layer commit...' && echo ''
 
@@ -140,11 +132,7 @@ RUN mkdir "$UE_GITDEPS" && echo '' && echo 'RUN directive complete. Docker will
 # When running with BuildKit, we use a cache mount to cache the dependency data across multiple build invocations
 WORKDIR [REDACTED_USER_HOME]/UnrealEngine
 RUN --mount=type=cache,target=[REDACTED_USER_HOME]/gitdeps,uid=1000,gid=1000 sudo apt-get update && \
-	./Setup.sh \
-		--exclude=Android \
-		--exclude=Mac \
-		--exclude=Win32 \
-		--exclude=Win64 && \
+	./Setup.sh --exclude=Android --exclude=Mac --exclude=Win32 --exclude=Win64 && \
 	sudo rm -rf /var/lib/apt/lists/* && echo '' && echo 'RUN directive complete. Docker will now commit the filesystem layer to disk.' && echo 'Note that for large filesystem layers this can take quite some time.' && echo 'Performing filesystem layer commit...' && echo ''
 
 FROM source as builder
@@ -160,7 +148,7 @@ RUN rm -rf [REDACTED_USER_HOME]/UnrealEngine/.git && echo '' && echo 'RUN directive complet
 # Ensure UBT is built before we create the Installed Build, since Build.sh explicitly sets the
 # target .NET Framework version, whereas InstalledEngineBuild.xml just uses the system default,
 # which can result in errors when running the built UBT due to the wrong version being targeted
-RUN ./Engine/Build/BatchFiles/Linux/Build.sh UnrealHeaderTool Linux Development -SkipBuild -buildubt && echo '' && echo 'RUN directive complete. Docker will now commit the filesystem layer to disk.' && echo 'Note that for large filesystem layers this can take quite some time.' && echo 'Performing filesystem layer commit...' && echo ''
+RUN ./Engine/Build/BatchFiles/Linux/Build.sh ShaderCompileWorker Linux Development -SkipBuild -buildubt && echo '' && echo 'RUN directive complete. Docker will now commit the filesystem layer to disk.' && echo 'Note that for large filesystem layers this can take quite some time.' && echo 'Performing filesystem layer commit...' && echo ''
 
 # Create an Installed Build of the Engine
 WORKDIR [REDACTED_USER_HOME]/UnrealEngine
@@ -187,7 +175,9 @@ WORKDIR [REDACTED_USER_HOME]/UnrealEngine
 # Perform first-run setup for Mono, UnrealBuildTool and AutomationTool, which makes it possible to build Unreal projects and plugins as users other than `ue4`
 # (Note that this will only work with 4.26.0 and newer, older Engine versions will always require write access to `[REDACTED_USER_HOME]/UnrealEngine`)
 # See the comments on this issue for details, including the need to ensure $HOME is set correctly: <[REDACTED_PROFILE_URL]/ue4-docker/issues/141>
-RUN ./Engine/Build/BatchFiles/Linux/Build.sh UnrealHeaderTool Linux Development -SkipBuild && \
+COPY print-editor-target.py /tmp/print-editor-target.py
+RUN EDITOR_TARGET=$(python3 /tmp/print-editor-target.py [REDACTED_USER_HOME]/UnrealEngine) && \
+	./Engine/Build/BatchFiles/Linux/Build.sh "$EDITOR_TARGET" Linux Development -SkipBuild && \
 	mkdir -p ./Engine/Programs/AutomationTool/Saved && \
 	chmod a+rw ./Engine/Programs/AutomationTool/Saved && echo '' && echo 'RUN directive complete. Docker will now commit the filesystem layer to disk.' && echo 'Note that for large filesystem layers this can take quite some time.' && echo 'Performing filesystem layer commit...' && echo ''
 
diff --git a/Engine/Extras/Containers/Dockerfiles/linux/dev-slim/print-editor-target.py b/Engine/Extras/Containers/Dockerfiles/linux/dev-slim/print-editor-target.py
new file mode 100644
--- /dev/null
+++ b/Engine/Extras/Containers/Dockerfiles/linux/dev-slim/print-editor-target.py
@@ -0,0 +1,12 @@
+#!/usr/bin/env python3
+from pathlib import Path
+import json, sys
+
+# Parse the Unreal Engine version information
+engineRoot = Path(sys.argv[1])
+versionFile = engineRoot / "Engine" / "Build" / "Build.version"
+versionDetails = json.loads(versionFile.read_text("utf-8"))
+
+# Determine the name of the Editor target based on the version
+target = "UE4Editor" if versionDetails["MajorVersion"] == 4 else "UnrealEditor"
+print(target)
diff --git a/Engine/Extras/Containers/Dockerfiles/linux/dev/Dockerfile b/Engine/Extras/Containers/Dockerfiles/linux/dev/Dockerfile
--- a/Engine/Extras/Containers/Dockerfiles/linux/dev/Dockerfile
+++ b/Engine/Extras/Containers/Dockerfiles/linux/dev/Dockerfile
@@ -1,12 +1,13 @@
-# This file was generated by ue4-docker version 0.0.102 with the following options:
+# This file was generated by ue4-docker version 0.0.110 with the following options:
 # 
 # - buildgraph_args: "-set:WithLinuxArm64=false -set:WithClient=true -set:WithServer=true"
 # - combine: true
 # - credential_mode: "secrets"
 # - disable_all_patches: true
 # - disable_labels: true
-# - enable_dso_patch: true
+# - enable_dso_patch: false
 # - excluded_components: {"ddc": false, "debug": false, "templates": false}
+# - gitdependencies_args: "--exclude=Android --exclude=Mac --exclude=Win32 --exclude=Win64"
 # 
 # This Dockerfile combines the steps for the following images:
 # 
@@ -82,15 +83,6 @@ RUN apt-get update && apt-get install -y --no-install-recommends \
 			xkb-data && \
 	rm -rf /var/lib/apt/lists/* && echo '' && echo 'RUN directive complete. Docker will now commit the filesystem layer to disk.' && echo 'Note that for large filesystem layers this can take quite some time.' && echo 'Performing filesystem layer commit...' && echo ''
 
-# Install the glibc DSO patch to improve Editor startup times
-RUN add-apt-repository -y ppa:slonopotamus/glibc-dso && \
-	apt-get update && \
-	apt upgrade -y libc6 && \
-	rm -rf /var/lib/apt/lists/* && echo '' && echo 'RUN directive complete. Docker will now commit the filesystem layer to disk.' && echo 'Note that for large filesystem layers this can take quite some time.' && echo 'Performing filesystem layer commit...' && echo ''
-
-# Enable the glibc DSO patch
-ENV GLIBC_TUNABLES=glibc.rtld.dynamic_sort=2
-
 # Disable the default "lecture" message the first time a user runs a command using sudo
 RUN echo 'Defaults lecture="never"' >> /etc/sudoers && echo '' && echo 'RUN directive complete. Docker will now commit the filesystem layer to disk.' && echo 'Note that for large filesystem layers this can take quite some time.' && echo 'Performing filesystem layer commit...' && echo ''
 
@@ -140,11 +132,7 @@ RUN mkdir "$UE_GITDEPS" && echo '' && echo 'RUN directive complete. Docker will
 # When running with BuildKit, we use a cache mount to cache the dependency data across multiple build invocations
 WORKDIR [REDACTED_USER_HOME]/UnrealEngine
 RUN --mount=type=cache,target=[REDACTED_USER_HOME]/gitdeps,uid=1000,gid=1000 sudo apt-get update && \
-	./Setup.sh \
-		--exclude=Android \
-		--exclude=Mac \
-		--exclude=Win32 \
-		--exclude=Win64 && \
+	./Setup.sh --exclude=Android --exclude=Mac --exclude=Win32 --exclude=Win64 && \
 	sudo rm -rf /var/lib/apt/lists/* && echo '' && echo 'RUN directive complete. Docker will now commit the filesystem layer to disk.' && echo 'Note that for large filesystem layers this can take quite some time.' && echo 'Performing filesystem layer commit...' && echo ''
 
 FROM source as builder
@@ -160,7 +148,7 @@ RUN rm -rf [REDACTED_USER_HOME]/UnrealEngine/.git && echo '' && echo 'RUN directive complet
 # Ensure UBT is built before we create the Installed Build, since Build.sh explicitly sets the
 # target .NET Framework version, whereas InstalledEngineBuild.xml just uses the system default,
 # which can result in errors when running the built UBT due to the wrong version being targeted
-RUN ./Engine/Build/BatchFiles/Linux/Build.sh UnrealHeaderTool Linux Development -SkipBuild -buildubt && echo '' && echo 'RUN directive complete. Docker will now commit the filesystem layer to disk.' && echo 'Note that for large filesystem layers this can take quite some time.' && echo 'Performing filesystem layer commit...' && echo ''
+RUN ./Engine/Build/BatchFiles/Linux/Build.sh ShaderCompileWorker Linux Development -SkipBuild -buildubt && echo '' && echo 'RUN directive complete. Docker will now commit the filesystem layer to disk.' && echo 'Note that for large filesystem layers this can take quite some time.' && echo 'Performing filesystem layer commit...' && echo ''
 
 # Create an Installed Build of the Engine
 WORKDIR [REDACTED_USER_HOME]/UnrealEngine
@@ -189,7 +177,9 @@ WORKDIR [REDACTED_USER_HOME]/UnrealEngine
 # Perform first-run setup for Mono, UnrealBuildTool and AutomationTool, which makes it possible to build Unreal projects and plugins as users other than `ue4`
 # (Note that this will only work with 4.26.0 and newer, older Engine versions will always require write access to `[REDACTED_USER_HOME]/UnrealEngine`)
 # See the comments on this issue for details, including the need to ensure $HOME is set correctly: <[REDACTED_PROFILE_URL]/ue4-docker/issues/141>
-RUN ./Engine/Build/BatchFiles/Linux/Build.sh UnrealHeaderTool Linux Development -SkipBuild && \
+COPY print-editor-target.py /tmp/print-editor-target.py
+RUN EDITOR_TARGET=$(python3 /tmp/print-editor-target.py [REDACTED_USER_HOME]/UnrealEngine) && \
+	./Engine/Build/BatchFiles/Linux/Build.sh "$EDITOR_TARGET" Linux Development -SkipBuild && \
 	mkdir -p ./Engine/Programs/AutomationTool/Saved && \
 	chmod a+rw ./Engine/Programs/AutomationTool/Saved && echo '' && echo 'RUN directive complete. Docker will now commit the filesystem layer to disk.' && echo 'Note that for large filesystem layers this can take quite some time.' && echo 'Performing filesystem layer commit...' && echo ''
 
diff --git a/Engine/Extras/Containers/Dockerfiles/linux/dev/print-editor-target.py b/Engine/Extras/Containers/Dockerfiles/linux/dev/print-editor-target.py
new file mode 100644
--- /dev/null
+++ b/Engine/Extras/Containers/Dockerfiles/linux/dev/print-editor-target.py
@@ -0,0 +1,12 @@
+#!/usr/bin/env python3
+from pathlib import Path
+import json, sys
+
+# Parse the Unreal Engine version information
+engineRoot = Path(sys.argv[1])
+versionFile = engineRoot / "Engine" / "Build" / "Build.version"
+versionDetails = json.loads(versionFile.read_text("utf-8"))
+
+# Determine the name of the Editor target based on the version
+target = "UE4Editor" if versionDetails["MajorVersion"] == 4 else "UnrealEditor"
+print(target)
diff --git a/Engine/Extras/Containers/Dockerfiles/linux/generate.sh b/Engine/Extras/Containers/Dockerfiles/linux/generate.sh
--- a/Engine/Extras/Containers/Dockerfiles/linux/generate.sh
+++ b/Engine/Extras/Containers/Dockerfiles/linux/generate.sh
@@ -42,8 +42,8 @@ function generateDockerfile {
 		# This strips out the image labels that are ordinarily used to facilitate cleanup and inspection of images built by ue4-docker
 		--opt disable-labels
 		
-		# This enables the glibc DSO sorting patch, which significantly improves Editor startup times
-		--opt enable-dso-patch
+		# This disables the glibc DSO sorting patch, since it is not needed under Ubuntu 22.04
+		--opt enable_dso_patch=false
 	)
 	
 	# Invoke ue4-docker to generate the Dockerfile
