Re-generate development Dockerfiles with ue4-docker version 0.0.112 and update for 5.4

diff --git a/Engine/Extras/Containers/Dockerfiles/linux/dev-slim/Dockerfile b/Engine/Extras/Containers/Dockerfiles/linux/dev-slim/Dockerfile
--- a/Engine/Extras/Containers/Dockerfiles/linux/dev-slim/Dockerfile
+++ b/Engine/Extras/Containers/Dockerfiles/linux/dev-slim/Dockerfile
@@ -1,11 +1,12 @@
-# This file was generated by ue4-docker version 0.0.110 with the following options:
+# This file was generated by ue4-docker version 0.0.112 with the following options:
 # 
-# - buildgraph_args: "-set:WithLinuxArm64=false -set:WithClient=true -set:WithServer=true"
+# - buildgraph_args: "-set:BuildIdOverride=UE_5.4 -set:WithLinuxArm64=false -set:WithClient=true -set:WithServer=true"
 # - combine: true
 # - credential_mode: "secrets"
 # - disable_all_patches: true
 # - disable_labels: true
 # - enable_dso_patch: false
+# - enable_ushell: true
 # - excluded_components: {"ddc": false, "debug": true, "templates": true}
 # - gitdependencies_args: "--exclude=Android --exclude=Mac --exclude=Win32 --exclude=Win64"
 # 
@@ -40,6 +41,7 @@ RUN apt-get update && apt-get install -y --no-install-recommends \
 		curl \
 		git \
 		git-lfs \
+		gpg-agent \
 		python3 \
 		python3-dev \
 		python3-pip \
@@ -49,6 +51,7 @@ RUN apt-get update && apt-get install -y --no-install-recommends \
 		tzdata \
 		unzip \
 		xdg-user-dirs \
+		xdg-utils \
 		zip && \
 	rm -rf /var/lib/apt/lists/* && echo '' && echo 'RUN directive complete. Docker will now commit the filesystem layer to disk.' && echo 'Note that for large filesystem layers this can take quite some time.' && echo 'Performing filesystem layer commit...' && echo ''
 
@@ -61,6 +64,7 @@ RUN apt-get update && apt-get install -y --no-install-recommends \
 			libcairo2 \
 			libfontconfig1 \
 			libfreetype6 \
+			libgbm1 \
 			libglu1 \
 			libnss3 \
 			libnspr4 \
@@ -92,6 +96,17 @@ RUN useradd --create-home --home [REDACTED_USER_HOME] --shell /bin/bash --uid 1000 ue4 && \
 	usermod -a -G audio,video,sudo ue4 && echo '' && echo 'RUN directive complete. Docker will now commit the filesystem layer to disk.' && echo 'Note that for large filesystem layers this can take quite some time.' && echo 'Performing filesystem layer commit...' && echo ''
 USER ue4
 
+# Install Python 3.12, which is required by ushell
+USER root
+RUN add-apt-repository -y ppa:deadsnakes/ppa && \
+	apt-get update && \
+	apt-get install -y --no-install-recommends python3.12 python3.12-venv && \
+	rm -rf /var/lib/apt/lists/* && echo '' && echo 'RUN directive complete. Docker will now commit the filesystem layer to disk.' && echo 'Note that for large filesystem layers this can take quite some time.' && echo 'Performing filesystem layer commit...' && echo ''
+USER ue4
+
+# Install a copy of pip for Python 3.12
+RUN curl -fsSL 'https://bootstrap.pypa.io/get-pip.py' | python3.12 && echo '' && echo 'RUN directive complete. Docker will now commit the filesystem layer to disk.' && echo 'Note that for large filesystem layers this can take quite some time.' && echo 'Performing filesystem layer commit...' && echo ''
+
 # Enable Git Large File Storage (LFS) support
 RUN git lfs install && echo '' && echo 'RUN directive complete. Docker will now commit the filesystem layer to disk.' && echo 'Note that for large filesystem layers this can take quite some time.' && echo 'Performing filesystem layer commit...' && echo ''
 
@@ -153,25 +168,49 @@ RUN ./Engine/Build/BatchFiles/Linux/Build.sh ShaderCompileWorker Linux Developme
 # Create an Installed Build of the Engine
 WORKDIR [REDACTED_USER_HOME]/UnrealEngine
 RUN ./Engine/Build/BatchFiles/RunUAT.sh BuildGraph \
-    -target="Make Installed Build Linux" \
-    -script=Engine/Build/InstalledEngineBuild.xml \
-    -set:HostPlatformOnly=true \
-    -set:WithDDC=true \
-    -set:WithLinuxArm64=false -set:WithClient=true -set:WithServer=true && \
+	-target="Make Installed Build Linux" \
+	-script=Engine/Build/InstalledEngineBuild.xml \
+	-set:HostPlatformOnly=true \
+	-set:WithDDC=true \
+	-set:BuildIdOverride=UE_5.4 -set:WithLinuxArm64=false -set:WithClient=true -set:WithServer=true && \
 	rm -R -f [REDACTED_USER_HOME]/UnrealEngine/LocalBuilds/InstalledDDC && echo '' && echo 'RUN directive complete. Docker will now commit the filesystem layer to disk.' && echo 'Note that for large filesystem layers this can take quite some time.' && echo 'Performing filesystem layer commit...' && echo ''
 
 # Split out components (DDC, debug symbols, template projects) so they can be copied into the final container image as separate filesystem layers
 COPY split-components.py /tmp/split-components.py
 RUN python3 /tmp/split-components.py [REDACTED_USER_HOME]/UnrealEngine/LocalBuilds/Engine/Linux [REDACTED_USER_HOME]/UnrealEngine/Components && echo '' && echo 'RUN directive complete. Docker will now commit the filesystem layer to disk.' && echo 'Note that for large filesystem layers this can take quite some time.' && echo 'Performing filesystem layer commit...' && echo ''
 
+# Ensure UnrealVersionSelector is built, since the prebuilt binaries may not be up-to-date
+RUN ./Engine/Build/BatchFiles/Linux/Build.sh UnrealVersionSelector Linux Shipping && echo '' && echo 'RUN directive complete. Docker will now commit the filesystem layer to disk.' && echo 'Note that for large filesystem layers this can take quite some time.' && echo 'Performing filesystem layer commit...' && echo ''
+
+# Copy InstalledBuild.txt from the Installed Build and run UnrealVersionSelector to populate Install.ini with any custom Build ID specified in the BuildGraph flags
+# (Note that the `-unattended` flag used below requires Unreal Engine 4.22 or newer, so this will break under older versions)
+# (Note also that custom Build IDs are supported by Unreal Engine 5.3.1 and newer, and older versions will just use a GUID as the Build ID)
+RUN cp [REDACTED_USER_HOME]/UnrealEngine/LocalBuilds/Engine/Linux/Engine/Build/InstalledBuild.txt [REDACTED_USER_HOME]/UnrealEngine/Engine/Build/InstalledBuild.txt && \
+	./Engine/Binaries/Linux/UnrealVersionSelector-Linux-Shipping -register -unattended && echo '' && echo 'RUN directive complete. Docker will now commit the filesystem layer to disk.' && echo 'Note that for large filesystem layers this can take quite some time.' && echo 'Performing filesystem layer commit...' && echo ''
+
+# Ensure ushell is copied to the Installed Build
+RUN rm -rf ./LocalBuilds/Engine/Linux/Engine/Extras/ushell && \
+	cp -r ./Engine/Extras/ushell ./LocalBuilds/Engine/Linux/Engine/Extras/ushell && \
+	bash -c 'set -e; shopt -s globstar; cd [REDACTED_USER_HOME]/UnrealEngine/LocalBuilds/Engine/Linux/Engine/Extras/ushell && chmod +x ./**/*.sh' && echo '' && echo 'RUN directive complete. Docker will now commit the filesystem layer to disk.' && echo 'Note that for large filesystem layers this can take quite some time.' && echo 'Performing filesystem layer commit...' && echo ''
+
 # Copy the Installed Build into a clean image, discarding the source build
 FROM prerequisites as minimal
 
 # Copy the Installed Build files from the builder image
 COPY --from=builder --chown=ue4:ue4 [REDACTED_USER_HOME]/UnrealEngine/LocalBuilds/Engine/Linux [REDACTED_USER_HOME]/UnrealEngine
 COPY --from=builder --chown=ue4:ue4 [REDACTED_USER_HOME]/UnrealEngine/Components/DDC [REDACTED_USER_HOME]/UnrealEngine
+
+# Copy Install.ini from the builder image, so it can be used by tools that read the list of engine installations (e.g. ushell)
+COPY --from=builder --chown=ue4:ue4 [REDACTED_USER_HOME]/.config/Epic/UnrealEngine/Install.ini [REDACTED_USER_HOME]/.config/Epic/UnrealEngine/Install.ini
 WORKDIR [REDACTED_USER_HOME]/UnrealEngine
 
+# Add ushell to the system PATH and alias `ushell` to `ushell.sh`
+ENV PATH="$PATH:[REDACTED_USER_HOME]/UnrealEngine/Engine/Extras/ushell"
+RUN echo 'alias ushell="ushell.sh"' >> [REDACTED_USER_HOME]/.bashrc && echo '' && echo 'RUN directive complete. Docker will now commit the filesystem layer to disk.' && echo 'Note that for large filesystem layers this can take quite some time.' && echo 'Performing filesystem layer commit...' && echo ''
+
+# Perform first-run setup for ushell
+RUN bash -c 'set -e; source [REDACTED_USER_HOME]/UnrealEngine/Engine/Extras/ushell/ushell.sh && exit 0' && echo '' && echo 'RUN directive complete. Docker will now commit the filesystem layer to disk.' && echo 'Note that for large filesystem layers this can take quite some time.' && echo 'Performing filesystem layer commit...' && echo ''
+
 # Perform first-run setup for Mono, UnrealBuildTool and AutomationTool, which makes it possible to build Unreal projects and plugins as users other than `ue4`
 # (Note that this will only work with 4.26.0 and newer, older Engine versions will always require write access to `[REDACTED_USER_HOME]/UnrealEngine`)
 # See the comments on this issue for details, including the need to ensure $HOME is set correctly: <[REDACTED_PROFILE_URL]/ue4-docker/issues/141>
diff --git a/Engine/Extras/Containers/Dockerfiles/linux/dev/Dockerfile b/Engine/Extras/Containers/Dockerfiles/linux/dev/Dockerfile
--- a/Engine/Extras/Containers/Dockerfiles/linux/dev/Dockerfile
+++ b/Engine/Extras/Containers/Dockerfiles/linux/dev/Dockerfile
@@ -1,11 +1,12 @@
-# This file was generated by ue4-docker version 0.0.110 with the following options:
+# This file was generated by ue4-docker version 0.0.112 with the following options:
 # 
-# - buildgraph_args: "-set:WithLinuxArm64=false -set:WithClient=true -set:WithServer=true"
+# - buildgraph_args: "-set:BuildIdOverride=UE_5.4 -set:WithLinuxArm64=false -set:WithClient=true -set:WithServer=true"
 # - combine: true
 # - credential_mode: "secrets"
 # - disable_all_patches: true
 # - disable_labels: true
 # - enable_dso_patch: false
+# - enable_ushell: true
 # - excluded_components: {"ddc": false, "debug": false, "templates": false}
 # - gitdependencies_args: "--exclude=Android --exclude=Mac --exclude=Win32 --exclude=Win64"
 # 
@@ -40,6 +41,7 @@ RUN apt-get update && apt-get install -y --no-install-recommends \
 		curl \
 		git \
 		git-lfs \
+		gpg-agent \
 		python3 \
 		python3-dev \
 		python3-pip \
@@ -49,6 +51,7 @@ RUN apt-get update && apt-get install -y --no-install-recommends \
 		tzdata \
 		unzip \
 		xdg-user-dirs \
+		xdg-utils \
 		zip && \
 	rm -rf /var/lib/apt/lists/* && echo '' && echo 'RUN directive complete. Docker will now commit the filesystem layer to disk.' && echo 'Note that for large filesystem layers this can take quite some time.' && echo 'Performing filesystem layer commit...' && echo ''
 
@@ -61,6 +64,7 @@ RUN apt-get update && apt-get install -y --no-install-recommends \
 			libcairo2 \
 			libfontconfig1 \
 			libfreetype6 \
+			libgbm1 \
 			libglu1 \
 			libnss3 \
 			libnspr4 \
@@ -92,6 +96,17 @@ RUN useradd --create-home --home [REDACTED_USER_HOME] --shell /bin/bash --uid 1000 ue4 && \
 	usermod -a -G audio,video,sudo ue4 && echo '' && echo 'RUN directive complete. Docker will now commit the filesystem layer to disk.' && echo 'Note that for large filesystem layers this can take quite some time.' && echo 'Performing filesystem layer commit...' && echo ''
 USER ue4
 
+# Install Python 3.12, which is required by ushell
+USER root
+RUN add-apt-repository -y ppa:deadsnakes/ppa && \
+	apt-get update && \
+	apt-get install -y --no-install-recommends python3.12 python3.12-venv && \
+	rm -rf /var/lib/apt/lists/* && echo '' && echo 'RUN directive complete. Docker will now commit the filesystem layer to disk.' && echo 'Note that for large filesystem layers this can take quite some time.' && echo 'Performing filesystem layer commit...' && echo ''
+USER ue4
+
+# Install a copy of pip for Python 3.12
+RUN curl -fsSL 'https://bootstrap.pypa.io/get-pip.py' | python3.12 && echo '' && echo 'RUN directive complete. Docker will now commit the filesystem layer to disk.' && echo 'Note that for large filesystem layers this can take quite some time.' && echo 'Performing filesystem layer commit...' && echo ''
+
 # Enable Git Large File Storage (LFS) support
 RUN git lfs install && echo '' && echo 'RUN directive complete. Docker will now commit the filesystem layer to disk.' && echo 'Note that for large filesystem layers this can take quite some time.' && echo 'Performing filesystem layer commit...' && echo ''
 
@@ -153,17 +168,31 @@ RUN ./Engine/Build/BatchFiles/Linux/Build.sh ShaderCompileWorker Linux Developme
 # Create an Installed Build of the Engine
 WORKDIR [REDACTED_USER_HOME]/UnrealEngine
 RUN ./Engine/Build/BatchFiles/RunUAT.sh BuildGraph \
-    -target="Make Installed Build Linux" \
-    -script=Engine/Build/InstalledEngineBuild.xml \
-    -set:HostPlatformOnly=true \
-    -set:WithDDC=true \
-    -set:WithLinuxArm64=false -set:WithClient=true -set:WithServer=true && \
+	-target="Make Installed Build Linux" \
+	-script=Engine/Build/InstalledEngineBuild.xml \
+	-set:HostPlatformOnly=true \
+	-set:WithDDC=true \
+	-set:BuildIdOverride=UE_5.4 -set:WithLinuxArm64=false -set:WithClient=true -set:WithServer=true && \
 	rm -R -f [REDACTED_USER_HOME]/UnrealEngine/LocalBuilds/InstalledDDC && echo '' && echo 'RUN directive complete. Docker will now commit the filesystem layer to disk.' && echo 'Note that for large filesystem layers this can take quite some time.' && echo 'Performing filesystem layer commit...' && echo ''
 
 # Split out components (DDC, debug symbols, template projects) so they can be copied into the final container image as separate filesystem layers
 COPY split-components.py /tmp/split-components.py
 RUN python3 /tmp/split-components.py [REDACTED_USER_HOME]/UnrealEngine/LocalBuilds/Engine/Linux [REDACTED_USER_HOME]/UnrealEngine/Components && echo '' && echo 'RUN directive complete. Docker will now commit the filesystem layer to disk.' && echo 'Note that for large filesystem layers this can take quite some time.' && echo 'Performing filesystem layer commit...' && echo ''
 
+# Ensure UnrealVersionSelector is built, since the prebuilt binaries may not be up-to-date
+RUN ./Engine/Build/BatchFiles/Linux/Build.sh UnrealVersionSelector Linux Shipping && echo '' && echo 'RUN directive complete. Docker will now commit the filesystem layer to disk.' && echo 'Note that for large filesystem layers this can take quite some time.' && echo 'Performing filesystem layer commit...' && echo ''
+
+# Copy InstalledBuild.txt from the Installed Build and run UnrealVersionSelector to populate Install.ini with any custom Build ID specified in the BuildGraph flags
+# (Note that the `-unattended` flag used below requires Unreal Engine 4.22 or newer, so this will break under older versions)
+# (Note also that custom Build IDs are supported by Unreal Engine 5.3.1 and newer, and older versions will just use a GUID as the Build ID)
+RUN cp [REDACTED_USER_HOME]/UnrealEngine/LocalBuilds/Engine/Linux/Engine/Build/InstalledBuild.txt [REDACTED_USER_HOME]/UnrealEngine/Engine/Build/InstalledBuild.txt && \
+	./Engine/Binaries/Linux/UnrealVersionSelector-Linux-Shipping -register -unattended && echo '' && echo 'RUN directive complete. Docker will now commit the filesystem layer to disk.' && echo 'Note that for large filesystem layers this can take quite some time.' && echo 'Performing filesystem layer commit...' && echo ''
+
+# Ensure ushell is copied to the Installed Build
+RUN rm -rf ./LocalBuilds/Engine/Linux/Engine/Extras/ushell && \
+	cp -r ./Engine/Extras/ushell ./LocalBuilds/Engine/Linux/Engine/Extras/ushell && \
+	bash -c 'set -e; shopt -s globstar; cd [REDACTED_USER_HOME]/UnrealEngine/LocalBuilds/Engine/Linux/Engine/Extras/ushell && chmod +x ./**/*.sh' && echo '' && echo 'RUN directive complete. Docker will now commit the filesystem layer to disk.' && echo 'Note that for large filesystem layers this can take quite some time.' && echo 'Performing filesystem layer commit...' && echo ''
+
 # Copy the Installed Build into a clean image, discarding the source build
 FROM prerequisites as minimal
 
@@ -172,8 +201,18 @@ COPY --from=builder --chown=ue4:ue4 [REDACTED_USER_HOME]/UnrealEngine/LocalBuilds/Engine/Li
 COPY --from=builder --chown=ue4:ue4 [REDACTED_USER_HOME]/UnrealEngine/Components/DDC [REDACTED_USER_HOME]/UnrealEngine
 COPY --from=builder --chown=ue4:ue4 [REDACTED_USER_HOME]/UnrealEngine/Components/DebugSymbols [REDACTED_USER_HOME]/UnrealEngine
 COPY --from=builder --chown=ue4:ue4 [REDACTED_USER_HOME]/UnrealEngine/Components/TemplatesAndSamples [REDACTED_USER_HOME]/UnrealEngine
+
+# Copy Install.ini from the builder image, so it can be used by tools that read the list of engine installations (e.g. ushell)
+COPY --from=builder --chown=ue4:ue4 [REDACTED_USER_HOME]/.config/Epic/UnrealEngine/Install.ini [REDACTED_USER_HOME]/.config/Epic/UnrealEngine/Install.ini
 WORKDIR [REDACTED_USER_HOME]/UnrealEngine
 
+# Add ushell to the system PATH and alias `ushell` to `ushell.sh`
+ENV PATH="$PATH:[REDACTED_USER_HOME]/UnrealEngine/Engine/Extras/ushell"
+RUN echo 'alias ushell="ushell.sh"' >> [REDACTED_USER_HOME]/.bashrc && echo '' && echo 'RUN directive complete. Docker will now commit the filesystem layer to disk.' && echo 'Note that for large filesystem layers this can take quite some time.' && echo 'Performing filesystem layer commit...' && echo ''
+
+# Perform first-run setup for ushell
+RUN bash -c 'set -e; source [REDACTED_USER_HOME]/UnrealEngine/Engine/Extras/ushell/ushell.sh && exit 0' && echo '' && echo 'RUN directive complete. Docker will now commit the filesystem layer to disk.' && echo 'Note that for large filesystem layers this can take quite some time.' && echo 'Performing filesystem layer commit...' && echo ''
+
 # Perform first-run setup for Mono, UnrealBuildTool and AutomationTool, which makes it possible to build Unreal projects and plugins as users other than `ue4`
 # (Note that this will only work with 4.26.0 and newer, older Engine versions will always require write access to `[REDACTED_USER_HOME]/UnrealEngine`)
 # See the comments on this issue for details, including the need to ensure $HOME is set correctly: <[REDACTED_PROFILE_URL]/ue4-docker/issues/141>
diff --git a/Engine/Extras/Containers/Dockerfiles/linux/generate.sh b/Engine/Extras/Containers/Dockerfiles/linux/generate.sh
--- a/Engine/Extras/Containers/Dockerfiles/linux/generate.sh
+++ b/Engine/Extras/Containers/Dockerfiles/linux/generate.sh
@@ -15,7 +15,7 @@ function generateDockerfile {
 		# Engine source code are defined by Docker build arguments that the user sets when building images with the generated Dockerfiles.
 		# The need to specify a release number is purely an artifact of how the `ue4-docker build` command was originally designed prior to
 		# the addition of functionality to generate Dockerfiles rather than always building container images.
-		'5.0.0'
+		--ue-version '5.0.0'
 		
 		# Ensure that we generate Dockerfiles for Linux container images even if this script is run on a Windows host system (e.g. in git bash)
 		# (Note that although generating the Dockerfiles under Windows is fine, it is NOT recommended that the images actually be built under Windows)
@@ -30,12 +30,15 @@ function generateDockerfile {
 		# This specifies the individual images whose Dockerfiles will be combined
 		--target=build-prerequisites,source,minimal
 		
-		# This disables building the Engine for AArch64 when creating an Installed Build, and enables support for client and dedicated server targets
-		--opt buildgraph-args='-set:WithLinuxArm64=false -set:WithClient=true -set:WithServer=true'
+		# This sets the Build ID, disables building the Engine for AArch64, and enables support for client and dedicated server targets
+		--opt buildgraph-args='-set:BuildIdOverride=UE_5.4 -set:WithLinuxArm64=false -set:WithClient=true -set:WithServer=true'
 		
 		# This enables the use of BuildKit build secrets, which is necessary in order to build images independently of ue4-docker itself
 		--opt credential-mode=secrets
 		
+		# This enables support for ushell
+		--opt enable-ushell=true
+		
 		# This strips out the Dockerfile code that ue4-docker would ordinarily generate that is uneccesary for Unreal Engine 4.26+
 		--opt disable-all-patches
 		
