Fix container build for 5.4

diff --git a/Engine/Extras/Containers/Dockerfiles/linux/dev-slim/Dockerfile b/Engine/Extras/Containers/Dockerfiles/linux/dev-slim/Dockerfile
--- a/Engine/Extras/Containers/Dockerfiles/linux/dev-slim/Dockerfile
+++ b/Engine/Extras/Containers/Dockerfiles/linux/dev-slim/Dockerfile
@@ -8,7 +8,7 @@
 # - enable_dso_patch: false
 # - enable_ushell: true
 # - excluded_components: {"ddc": false, "debug": true, "templates": true}
-# - gitdependencies_args: "--exclude=Android --exclude=Mac --exclude=Win32 --exclude=Win64"
+# - gitdependencies_args: ""
 # 
 # This Dockerfile combines the steps for the following images:
 # 
@@ -147,7 +147,7 @@ RUN mkdir "$UE_GITDEPS" && echo '' && echo 'RUN directive complete. Docker will
 # When running with BuildKit, we use a cache mount to cache the dependency data across multiple build invocations
 WORKDIR [REDACTED_USER_HOME]/UnrealEngine
 RUN --mount=type=cache,target=[REDACTED_USER_HOME]/gitdeps,uid=1000,gid=1000 sudo apt-get update && \
-	./Setup.sh --exclude=Android --exclude=Mac --exclude=Win32 --exclude=Win64 && \
+	./Setup.sh  && \
 	sudo rm -rf /var/lib/apt/lists/* && echo '' && echo 'RUN directive complete. Docker will now commit the filesystem layer to disk.' && echo 'Note that for large filesystem layers this can take quite some time.' && echo 'Performing filesystem layer commit...' && echo ''
 
 FROM source as builder
diff --git a/Engine/Extras/Containers/Dockerfiles/linux/dev/Dockerfile b/Engine/Extras/Containers/Dockerfiles/linux/dev/Dockerfile
--- a/Engine/Extras/Containers/Dockerfiles/linux/dev/Dockerfile
+++ b/Engine/Extras/Containers/Dockerfiles/linux/dev/Dockerfile
@@ -8,7 +8,7 @@
 # - enable_dso_patch: false
 # - enable_ushell: true
 # - excluded_components: {"ddc": false, "debug": false, "templates": false}
-# - gitdependencies_args: "--exclude=Android --exclude=Mac --exclude=Win32 --exclude=Win64"
+# - gitdependencies_args: ""
 # 
 # This Dockerfile combines the steps for the following images:
 # 
@@ -147,7 +147,7 @@ RUN mkdir "$UE_GITDEPS" && echo '' && echo 'RUN directive complete. Docker will
 # When running with BuildKit, we use a cache mount to cache the dependency data across multiple build invocations
 WORKDIR [REDACTED_USER_HOME]/UnrealEngine
 RUN --mount=type=cache,target=[REDACTED_USER_HOME]/gitdeps,uid=1000,gid=1000 sudo apt-get update && \
-	./Setup.sh --exclude=Android --exclude=Mac --exclude=Win32 --exclude=Win64 && \
+	./Setup.sh  && \
 	sudo rm -rf /var/lib/apt/lists/* && echo '' && echo 'RUN directive complete. Docker will now commit the filesystem layer to disk.' && echo 'Note that for large filesystem layers this can take quite some time.' && echo 'Performing filesystem layer commit...' && echo ''
 
 FROM source as builder
diff --git a/Engine/Extras/Containers/Dockerfiles/linux/generate.sh b/Engine/Extras/Containers/Dockerfiles/linux/generate.sh
--- a/Engine/Extras/Containers/Dockerfiles/linux/generate.sh
+++ b/Engine/Extras/Containers/Dockerfiles/linux/generate.sh
@@ -33,6 +33,9 @@ function generateDockerfile {
 		# This sets the Build ID, disables building the Engine for AArch64, and enables support for client and dedicated server targets
 		--opt buildgraph-args='-set:BuildIdOverride=UE_5.4 -set:WithLinuxArm64=false -set:WithClient=true -set:WithServer=true'
 		
+		# This ensures we do not exclude binary files for any platforms
+		--opt gitdependencies-args=''
+		
 		# This enables the use of BuildKit build secrets, which is necessary in order to build images independently of ue4-docker itself
 		--opt credential-mode=secrets
 		
