Update Windows runtime image for 5.4 using latest DLL list from ue4-docker

diff --git a/Engine/Extras/Containers/Dockerfiles/windows/build.ps1 b/Engine/Extras/Containers/Dockerfiles/windows/build.ps1
--- a/Engine/Extras/Containers/Dockerfiles/windows/build.ps1
+++ b/Engine/Extras/Containers/Dockerfiles/windows/build.ps1
@@ -15,22 +15,28 @@ if ($release -ne "")
 }
 else
 {
-	# Retrieve the Windows release number (e.g. 1903, 1909, 2004, 20H2, etc.)
-	$displayVersion = (Get-ItemProperty -Path 'Registry::HKEY_LOCAL_MACHINE\SOFTWARE\Microsoft\Windows NT\CurrentVersion' -Name DisplayVersion -ErrorAction SilentlyContinue)
-	if ($displayVersion) {
-		$windowsRelease = $displayVersion.DisplayVersion
-	} else {
-		$windowsRelease = (Get-ItemProperty -Path 'Registry::HKEY_LOCAL_MACHINE\SOFTWARE\Microsoft\Windows NT\CurrentVersion' -Name ReleaseId).ReleaseId
-	}
-	
-	# Use LTSC2022 images under Windows 11 and Windows Server 2022
+	# Determine whether we are running under Windows Server 2022 / Windows 11, or an older version of Windows
 	[int]$kernelBuild = (Get-ItemProperty -Path 'Registry::HKEY_LOCAL_MACHINE\SOFTWARE\Microsoft\Windows NT\CurrentVersion' -Name CurrentBuildNumber).CurrentBuildNumber
-	if ($kernelBuild -ge 20348) {
+	if ($kernelBuild -ge 20348)
+	{
+		# Use a Windows Server 2022 base image
 		$windowsRelease = "ltsc2022"
+		
+		# Use process isolation mode for improved performance
+		$isolation = "process"
+	}
+	else
+	{
+		# Use a Windows Server 2019 base image
+		$windowsRelease = "ltsc2019"
+		
+		# Use process isolation mode if we're running under Windows Server 2019 itself, otherwise use Hyper-V isolation mode (e.g. when running under Windows 10)
+		if ($kernelBuild -eq 17763) {
+			$isolation = "process"
+		} else {
+			$isolation = "hyperv"
+		}
 	}
-
-	# Use process isolation mode for improved performance
-	$isolation = "process"
 	
 	# Don't suffix the image tag
 	$tag = "runtime-windows"
@@ -42,6 +48,12 @@ if ($windowsRelease -eq "ltsc2022") {
 	$dllImage = "mcr.microsoft.com/windows/server"
 }
 
+# Identify any DLL files that are required for the specific version of Windows
+$versionSpecificDLLs = ""
+if ($windowsRelease -eq "ltsc2019") {
+	$versionSpecificDLLs = "ksuser.dll"
+}
+
 # Build our runtime container image using the correct base image for the selected Windows version
 "Building runtime container image for Windows version $windowsRelease with ``$isolation`` isolation mode and DLL files from ``$dllImage``..."
-docker build -t "ghcr.io/epicgames/unreal-engine:$tag" --isolation="$isolation" --build-arg "DLL_IMAGE=$dllImage" --build-arg "BASETAG=$windowsRelease" ./runtime
+docker build -t "ghcr.io/epicgames/unreal-engine:$tag" --isolation="$isolation" --build-arg "DLL_IMAGE=$dllImage" --build-arg "BASETAG=$windowsRelease" --build-arg "VERSION_SPECIFIC_DLLS=$versionSpecificDLLs" ./runtime
diff --git a/Engine/Extras/Containers/Dockerfiles/windows/runtime/Dockerfile b/Engine/Extras/Containers/Dockerfiles/windows/runtime/Dockerfile
--- a/Engine/Extras/Containers/Dockerfiles/windows/runtime/Dockerfile
+++ b/Engine/Extras/Containers/Dockerfiles/windows/runtime/Dockerfile
@@ -2,6 +2,7 @@
 ARG DLL_IMAGE
 ARG BASETAG
 FROM ${DLL_IMAGE}:${BASETAG} AS full
+ARG VERSION_SPECIFIC_DLLS
 
 # Gather the system DLLs that we need from the full Windows base image
 RUN xcopy /y C:\Windows\System32\avicap32.dll C:\GatheredDlls\ && `
@@ -12,7 +13,7 @@ RUN xcopy /y C:\Windows\System32\avicap32.dll C:\GatheredDlls\ && `
 	xcopy /y C:\Windows\System32\dsound.dll C:\GatheredDlls\ && `
 	xcopy /y C:\Windows\System32\dxva2.dll C:\GatheredDlls\ && `
 	xcopy /y C:\Windows\System32\glu32.dll C:\GatheredDlls\ && `
-	xcopy /y C:\Windows\System32\ksuser.dll C:\GatheredDlls\ && `
+	xcopy /y C:\Windows\System32\InputHost.dll C:\GatheredDlls\ && `
 	xcopy /y C:\Windows\System32\mf.dll C:\GatheredDlls\ && `
 	xcopy /y C:\Windows\System32\mfcore.dll C:\GatheredDlls\ && `
 	xcopy /y C:\Windows\System32\mfplat.dll C:\GatheredDlls\ && `
@@ -23,7 +24,9 @@ RUN xcopy /y C:\Windows\System32\avicap32.dll C:\GatheredDlls\ && `
 	xcopy /y C:\Windows\System32\msvfw32.dll C:\GatheredDlls\ && `
 	xcopy /y C:\Windows\System32\opengl32.dll C:\GatheredDlls\ && `
 	xcopy /y C:\Windows\System32\ResampleDMO.dll C:\GatheredDlls\ && `
-	xcopy /y C:\Windows\System32\ResourcePolicyClient.dll C:\GatheredDlls\
+	xcopy /y C:\Windows\System32\ResourcePolicyClient.dll C:\GatheredDlls\ && `
+	xcopy /y C:\Windows\System32\XInput1_4.dll C:\GatheredDlls\ && `
+	powershell -Command "$dlls = $env:VERSION_SPECIFIC_DLLS; if ($dlls) {foreach ($dll in $dlls.Split(';')) { & xcopy /y C:\Windows\System32\$dll C:\GatheredDlls\ } }"
 
 # Retrieve the DirectX runtime files required by the Unreal Engine, since even the full Windows base image does not include them
 RUN curl --progress-bar -L "https://download.microsoft.com/download/8/4/A/84A35BF1-DAFE-4AE8-82AF-AD2AE20B6B14/directx_Jun2010_redist.exe" --output %TEMP%\directx_redist.exe && `
@@ -46,7 +49,7 @@ FROM mcr.microsoft.com/windows/servercore:${BASETAG}
 COPY --from=full C:\GatheredDlls\ C:\Windows\System32\
 
 # Install the Visual C++ runtime files using Chocolatey
-RUN powershell -NoProfile -ExecutionPolicy Bypass -Command "[System.Net.ServicePointManager]::SecurityProtocol = [System.Net.ServicePointManager]::SecurityProtocol -bor 3072; iex ((New-Object System.Net.WebClient).DownloadString('https://chocolatey.org/install.ps1'))"
+RUN powershell -NoProfile -ExecutionPolicy Bypass -Command "$env:chocolateyVersion = '1.4.0'; [System.Net.ServicePointManager]::SecurityProtocol = [System.Net.ServicePointManager]::SecurityProtocol -bor 3072; iex ((New-Object System.Net.WebClient).DownloadString('https://chocolatey.org/install.ps1'))"
 RUN choco install -y vcredist-all
 
 # Copy our entrypoint script for enabling vendor-specific graphics APIs
