Re-generate development Dockerfiles with ue4-docker version 0.0.114

diff --git a/Engine/Extras/Containers/Dockerfiles/linux/dev-slim/Dockerfile b/Engine/Extras/Containers/Dockerfiles/linux/dev-slim/Dockerfile
--- a/Engine/Extras/Containers/Dockerfiles/linux/dev-slim/Dockerfile
+++ b/Engine/Extras/Containers/Dockerfiles/linux/dev-slim/Dockerfile
@@ -1,4 +1,4 @@
-# This file was generated by ue4-docker version 0.0.112 with the following options:
+# This file was generated by ue4-docker version 0.0.114 with the following options:
 # 
 # - buildgraph_args: "-set:BuildIdOverride=UE_5.4 -set:WithLinuxArm64=false -set:WithClient=true -set:WithServer=true"
 # - combine: true
@@ -175,10 +175,6 @@ RUN ./Engine/Build/BatchFiles/RunUAT.sh BuildGraph \
 	-set:BuildIdOverride=UE_5.4 -set:WithLinuxArm64=false -set:WithClient=true -set:WithServer=true && \
 	rm -R -f [REDACTED_USER_HOME]/UnrealEngine/LocalBuilds/InstalledDDC && echo '' && echo 'RUN directive complete. Docker will now commit the filesystem layer to disk.' && echo 'Note that for large filesystem layers this can take quite some time.' && echo 'Performing filesystem layer commit...' && echo ''
 
-# Split out components (DDC, debug symbols, template projects) so they can be copied into the final container image as separate filesystem layers
-COPY split-components.py /tmp/split-components.py
-RUN python3 /tmp/split-components.py [REDACTED_USER_HOME]/UnrealEngine/LocalBuilds/Engine/Linux [REDACTED_USER_HOME]/UnrealEngine/Components && echo '' && echo 'RUN directive complete. Docker will now commit the filesystem layer to disk.' && echo 'Note that for large filesystem layers this can take quite some time.' && echo 'Performing filesystem layer commit...' && echo ''
-
 # Ensure UnrealVersionSelector is built, since the prebuilt binaries may not be up-to-date
 RUN ./Engine/Build/BatchFiles/Linux/Build.sh UnrealVersionSelector Linux Shipping && echo '' && echo 'RUN directive complete. Docker will now commit the filesystem layer to disk.' && echo 'Note that for large filesystem layers this can take quite some time.' && echo 'Performing filesystem layer commit...' && echo ''
 
@@ -193,11 +189,22 @@ RUN rm -rf ./LocalBuilds/Engine/Linux/Engine/Extras/ushell && \
 	cp -r ./Engine/Extras/ushell ./LocalBuilds/Engine/Linux/Engine/Extras/ushell && \
 	bash -c 'set -e; shopt -s globstar; cd [REDACTED_USER_HOME]/UnrealEngine/LocalBuilds/Engine/Linux/Engine/Extras/ushell && chmod +x ./**/*.sh' && echo '' && echo 'RUN directive complete. Docker will now commit the filesystem layer to disk.' && echo 'Note that for large filesystem layers this can take quite some time.' && echo 'Performing filesystem layer commit...' && echo ''
 
+# Split out both optional components (DDC, debug symbols, template projects) and large subdirectories so they can be copied
+# into the final container image as separate filesystem layers, avoiding creating a single monolithic layer with everything
+COPY split-components.py /tmp/split-components.py
+RUN python3 /tmp/split-components.py [REDACTED_USER_HOME]/UnrealEngine/LocalBuilds/Engine/Linux [REDACTED_USER_HOME]/UnrealEngine/Components && echo '' && echo 'RUN directive complete. Docker will now commit the filesystem layer to disk.' && echo 'Note that for large filesystem layers this can take quite some time.' && echo 'Performing filesystem layer commit...' && echo ''
+
 # Copy the Installed Build into a clean image, discarding the source build
 FROM prerequisites as minimal
 
 # Copy the Installed Build files from the builder image
 COPY --from=builder --chown=ue4:ue4 [REDACTED_USER_HOME]/UnrealEngine/LocalBuilds/Engine/Linux [REDACTED_USER_HOME]/UnrealEngine
+COPY --from=builder --chown=ue4:ue4 [REDACTED_USER_HOME]/UnrealEngine/Components/Binaries [REDACTED_USER_HOME]/UnrealEngine
+COPY --from=builder --chown=ue4:ue4 [REDACTED_USER_HOME]/UnrealEngine/Components/Content [REDACTED_USER_HOME]/UnrealEngine
+COPY --from=builder --chown=ue4:ue4 [REDACTED_USER_HOME]/UnrealEngine/Components/Extras [REDACTED_USER_HOME]/UnrealEngine
+COPY --from=builder --chown=ue4:ue4 [REDACTED_USER_HOME]/UnrealEngine/Components/Intermediate [REDACTED_USER_HOME]/UnrealEngine
+COPY --from=builder --chown=ue4:ue4 [REDACTED_USER_HOME]/UnrealEngine/Components/Plugins [REDACTED_USER_HOME]/UnrealEngine
+COPY --from=builder --chown=ue4:ue4 [REDACTED_USER_HOME]/UnrealEngine/Components/Source [REDACTED_USER_HOME]/UnrealEngine
 COPY --from=builder --chown=ue4:ue4 [REDACTED_USER_HOME]/UnrealEngine/Components/DDC [REDACTED_USER_HOME]/UnrealEngine
 
 # Copy Install.ini from the builder image, so it can be used by tools that read the list of engine installations (e.g. ushell)
diff --git a/Engine/Extras/Containers/Dockerfiles/linux/dev-slim/split-components.py b/Engine/Extras/Containers/Dockerfiles/linux/dev-slim/split-components.py
--- a/Engine/Extras/Containers/Dockerfiles/linux/dev-slim/split-components.py
+++ b/Engine/Extras/Containers/Dockerfiles/linux/dev-slim/split-components.py
@@ -58,3 +58,13 @@ subdirs = [join(rootDir, subdir) for subdir in ["FeaturePacks", "Samples", "Temp
 extractComponent(
     rootDir, outputDir, "TemplatesAndSamples", "template projects and samples", subdirs
 )
+
+# Extract the larger non-optional subdirectories of the Engine directory
+for subdir in ["Binaries", "Content", "Extras", "Intermediate", "Plugins", "Source"]:
+    extractComponent(
+        rootDir,
+        outputDir,
+        subdir,
+        f"{subdir} subdirectory",
+        [join(rootDir, "Engine", subdir)],
+    )
diff --git a/Engine/Extras/Containers/Dockerfiles/linux/dev/Dockerfile b/Engine/Extras/Containers/Dockerfiles/linux/dev/Dockerfile
--- a/Engine/Extras/Containers/Dockerfiles/linux/dev/Dockerfile
+++ b/Engine/Extras/Containers/Dockerfiles/linux/dev/Dockerfile
@@ -1,4 +1,4 @@
-# This file was generated by ue4-docker version 0.0.112 with the following options:
+# This file was generated by ue4-docker version 0.0.114 with the following options:
 # 
 # - buildgraph_args: "-set:BuildIdOverride=UE_5.4 -set:WithLinuxArm64=false -set:WithClient=true -set:WithServer=true"
 # - combine: true
@@ -175,10 +175,6 @@ RUN ./Engine/Build/BatchFiles/RunUAT.sh BuildGraph \
 	-set:BuildIdOverride=UE_5.4 -set:WithLinuxArm64=false -set:WithClient=true -set:WithServer=true && \
 	rm -R -f [REDACTED_USER_HOME]/UnrealEngine/LocalBuilds/InstalledDDC && echo '' && echo 'RUN directive complete. Docker will now commit the filesystem layer to disk.' && echo 'Note that for large filesystem layers this can take quite some time.' && echo 'Performing filesystem layer commit...' && echo ''
 
-# Split out components (DDC, debug symbols, template projects) so they can be copied into the final container image as separate filesystem layers
-COPY split-components.py /tmp/split-components.py
-RUN python3 /tmp/split-components.py [REDACTED_USER_HOME]/UnrealEngine/LocalBuilds/Engine/Linux [REDACTED_USER_HOME]/UnrealEngine/Components && echo '' && echo 'RUN directive complete. Docker will now commit the filesystem layer to disk.' && echo 'Note that for large filesystem layers this can take quite some time.' && echo 'Performing filesystem layer commit...' && echo ''
-
 # Ensure UnrealVersionSelector is built, since the prebuilt binaries may not be up-to-date
 RUN ./Engine/Build/BatchFiles/Linux/Build.sh UnrealVersionSelector Linux Shipping && echo '' && echo 'RUN directive complete. Docker will now commit the filesystem layer to disk.' && echo 'Note that for large filesystem layers this can take quite some time.' && echo 'Performing filesystem layer commit...' && echo ''
 
@@ -193,11 +189,22 @@ RUN rm -rf ./LocalBuilds/Engine/Linux/Engine/Extras/ushell && \
 	cp -r ./Engine/Extras/ushell ./LocalBuilds/Engine/Linux/Engine/Extras/ushell && \
 	bash -c 'set -e; shopt -s globstar; cd [REDACTED_USER_HOME]/UnrealEngine/LocalBuilds/Engine/Linux/Engine/Extras/ushell && chmod +x ./**/*.sh' && echo '' && echo 'RUN directive complete. Docker will now commit the filesystem layer to disk.' && echo 'Note that for large filesystem layers this can take quite some time.' && echo 'Performing filesystem layer commit...' && echo ''
 
+# Split out both optional components (DDC, debug symbols, template projects) and large subdirectories so they can be copied
+# into the final container image as separate filesystem layers, avoiding creating a single monolithic layer with everything
+COPY split-components.py /tmp/split-components.py
+RUN python3 /tmp/split-components.py [REDACTED_USER_HOME]/UnrealEngine/LocalBuilds/Engine/Linux [REDACTED_USER_HOME]/UnrealEngine/Components && echo '' && echo 'RUN directive complete. Docker will now commit the filesystem layer to disk.' && echo 'Note that for large filesystem layers this can take quite some time.' && echo 'Performing filesystem layer commit...' && echo ''
+
 # Copy the Installed Build into a clean image, discarding the source build
 FROM prerequisites as minimal
 
 # Copy the Installed Build files from the builder image
 COPY --from=builder --chown=ue4:ue4 [REDACTED_USER_HOME]/UnrealEngine/LocalBuilds/Engine/Linux [REDACTED_USER_HOME]/UnrealEngine
+COPY --from=builder --chown=ue4:ue4 [REDACTED_USER_HOME]/UnrealEngine/Components/Binaries [REDACTED_USER_HOME]/UnrealEngine
+COPY --from=builder --chown=ue4:ue4 [REDACTED_USER_HOME]/UnrealEngine/Components/Content [REDACTED_USER_HOME]/UnrealEngine
+COPY --from=builder --chown=ue4:ue4 [REDACTED_USER_HOME]/UnrealEngine/Components/Extras [REDACTED_USER_HOME]/UnrealEngine
+COPY --from=builder --chown=ue4:ue4 [REDACTED_USER_HOME]/UnrealEngine/Components/Intermediate [REDACTED_USER_HOME]/UnrealEngine
+COPY --from=builder --chown=ue4:ue4 [REDACTED_USER_HOME]/UnrealEngine/Components/Plugins [REDACTED_USER_HOME]/UnrealEngine
+COPY --from=builder --chown=ue4:ue4 [REDACTED_USER_HOME]/UnrealEngine/Components/Source [REDACTED_USER_HOME]/UnrealEngine
 COPY --from=builder --chown=ue4:ue4 [REDACTED_USER_HOME]/UnrealEngine/Components/DDC [REDACTED_USER_HOME]/UnrealEngine
 COPY --from=builder --chown=ue4:ue4 [REDACTED_USER_HOME]/UnrealEngine/Components/DebugSymbols [REDACTED_USER_HOME]/UnrealEngine
 COPY --from=builder --chown=ue4:ue4 [REDACTED_USER_HOME]/UnrealEngine/Components/TemplatesAndSamples [REDACTED_USER_HOME]/UnrealEngine
diff --git a/Engine/Extras/Containers/Dockerfiles/linux/dev/split-components.py b/Engine/Extras/Containers/Dockerfiles/linux/dev/split-components.py
--- a/Engine/Extras/Containers/Dockerfiles/linux/dev/split-components.py
+++ b/Engine/Extras/Containers/Dockerfiles/linux/dev/split-components.py
@@ -58,3 +58,13 @@ subdirs = [join(rootDir, subdir) for subdir in ["FeaturePacks", "Samples", "Temp
 extractComponent(
     rootDir, outputDir, "TemplatesAndSamples", "template projects and samples", subdirs
 )
+
+# Extract the larger non-optional subdirectories of the Engine directory
+for subdir in ["Binaries", "Content", "Extras", "Intermediate", "Plugins", "Source"]:
+    extractComponent(
+        rootDir,
+        outputDir,
+        subdir,
+        f"{subdir} subdirectory",
+        [join(rootDir, "Engine", subdir)],
+    )
