Re-generate development Dockerfiles with ue4-docker version 0.0.119

#rnx

diff --git a/Engine/Extras/Containers/Dockerfiles/linux/dev-slim/Dockerfile b/Engine/Extras/Containers/Dockerfiles/linux/dev-slim/Dockerfile
--- a/Engine/Extras/Containers/Dockerfiles/linux/dev-slim/Dockerfile
+++ b/Engine/Extras/Containers/Dockerfiles/linux/dev-slim/Dockerfile
@@ -1,6 +1,7 @@
-# This file was generated by ue4-docker version 0.0.114 with the following options:
+# This file was generated by ue4-docker version 0.0.119 with the following options:
 # 
-# - buildgraph_args: "-set:BuildIdOverride=UE_5.4 -set:WithLinuxArm64=false -set:WithClient=true -set:WithServer=true"
+# - automatic_build_id: true
+# - buildgraph_args: "-set:WithLinuxArm64=false -set:WithClient=true -set:WithServer=true"
 # - combine: true
 # - credential_mode: "secrets"
 # - disable_all_patches: true
@@ -17,22 +18,22 @@
 # - ue4-minimal
 
 ARG BASEIMAGE
-FROM ${BASEIMAGE} as prerequisites
+FROM ${BASEIMAGE} AS prerequisites
 
 # Disable interactive prompts during package installation
 ENV DEBIAN_FRONTEND=noninteractive
 
 # Enable CUDA support for NVIDIA GPUs (even when not using a CUDA base image), since evidently some versions of UE unconditionally assume
 # `libcuda.so.1` exists when using the NVIDIA proprietary drivers, and will fail to initialise the Vulkan RHI if it is missing
-ENV NVIDIA_DRIVER_CAPABILITIES ${NVIDIA_DRIVER_CAPABILITIES},compute
+ENV NVIDIA_DRIVER_CAPABILITIES=${NVIDIA_DRIVER_CAPABILITIES},compute
 
 # Add the "display" driver capability for NVIDIA GPUs
 # (This allows us to run the Editor from an interactive container by bind-mounting the host system's X11 socket)
-ENV NVIDIA_DRIVER_CAPABILITIES ${NVIDIA_DRIVER_CAPABILITIES},display
+ENV NVIDIA_DRIVER_CAPABILITIES=${NVIDIA_DRIVER_CAPABILITIES},display
 
 # Enable NVENC support for use by Unreal Engine plugins that depend on it (e.g. Pixel Streaming)
 # (Note that adding `video` seems to implicitly enable `compute` as well, but we include separate directives here to clearly indicate the purpose of both)
-ENV NVIDIA_DRIVER_CAPABILITIES ${NVIDIA_DRIVER_CAPABILITIES},video
+ENV NVIDIA_DRIVER_CAPABILITIES=${NVIDIA_DRIVER_CAPABILITIES},video
 
 # Install our build prerequisites
 RUN apt-get update && apt-get install -y --no-install-recommends \
@@ -42,6 +43,7 @@ RUN apt-get update && apt-get install -y --no-install-recommends \
 		git \
 		git-lfs \
 		gpg-agent \
+		jq \
 		python3 \
 		python3-dev \
 		python3-pip \
@@ -87,6 +89,12 @@ RUN apt-get update && apt-get install -y --no-install-recommends \
 			xkb-data && \
 	rm -rf /var/lib/apt/lists/* && echo '' && echo 'RUN directive complete. Docker will now commit the filesystem layer to disk.' && echo 'Note that for large filesystem layers this can take quite some time.' && echo 'Performing filesystem layer commit...' && echo ''
 
+# Install the Avahi runtime libraries required by the NDIMedia plugin
+# (Starting in Unreal Engine 5.7, we need these installed before creating an Installed Build to prevent cooking failures)
+RUN apt-get update && apt-get install -y --no-install-recommends \
+			libavahi-client3 && \
+	rm -rf /var/lib/apt/lists/* && echo '' && echo 'RUN directive complete. Docker will now commit the filesystem layer to disk.' && echo 'Note that for large filesystem layers this can take quite some time.' && echo 'Performing filesystem layer commit...' && echo ''
+
 # Disable the default "lecture" message the first time a user runs a command using sudo
 RUN echo 'Defaults lecture="never"' >> /etc/sudoers && echo '' && echo 'RUN directive complete. Docker will now commit the filesystem layer to disk.' && echo 'Note that for large filesystem layers this can take quite some time.' && echo 'Performing filesystem layer commit...' && echo ''
 
@@ -150,29 +158,39 @@ RUN --mount=type=cache,target=[REDACTED_USER_HOME]/gitdeps,uid=1000,gid=1000 sudo apt-get u
 	./Setup.sh  && \
 	sudo rm -rf /var/lib/apt/lists/* && echo '' && echo 'RUN directive complete. Docker will now commit the filesystem layer to disk.' && echo 'Note that for large filesystem layers this can take quite some time.' && echo 'Performing filesystem layer commit...' && echo ''
 
-FROM source as builder
-
 # Set the changelist number in Build.version to ensure our Build ID is generated correctly
 ARG CHANGELIST
 COPY set-changelist.py /tmp/set-changelist.py
 RUN python3 /tmp/set-changelist.py [REDACTED_USER_HOME]/UnrealEngine/Engine/Build/Build.version $CHANGELIST && echo '' && echo 'RUN directive complete. Docker will now commit the filesystem layer to disk.' && echo 'Note that for large filesystem layers this can take quite some time.' && echo 'Performing filesystem layer commit...' && echo ''
 
+FROM source as builder
+
 # Remove the .git directory to disable UBT `git status` calls and speed up the build process
 RUN rm -rf [REDACTED_USER_HOME]/UnrealEngine/.git && echo '' && echo 'RUN directive complete. Docker will now commit the filesystem layer to disk.' && echo 'Note that for large filesystem layers this can take quite some time.' && echo 'Performing filesystem layer commit...' && echo ''
 
+# Disable Nuget package auditing (warnings NU1902 and NU1903)
+ENV NuGetAudit=false
+
 # Ensure UBT is built before we create the Installed Build, since Build.sh explicitly sets the
 # target .NET Framework version, whereas InstalledEngineBuild.xml just uses the system default,
 # which can result in errors when running the built UBT due to the wrong version being targeted
-RUN ./Engine/Build/BatchFiles/Linux/Build.sh ShaderCompileWorker Linux Development -SkipBuild -buildubt && echo '' && echo 'RUN directive complete. Docker will now commit the filesystem layer to disk.' && echo 'Note that for large filesystem layers this can take quite some time.' && echo 'Performing filesystem layer commit...' && echo ''
+RUN if [ -f ./Engine/Build/BatchFiles/BuildUBT.sh ]; then \
+		./Engine/Build/BatchFiles/BuildUBT.sh; \
+	else \
+		./Engine/Build/BatchFiles/Linux/Build.sh ShaderCompileWorker Linux Development -SkipBuild -buildubt; \
+	fi && echo '' && echo 'RUN directive complete. Docker will now commit the filesystem layer to disk.' && echo 'Note that for large filesystem layers this can take quite some time.' && echo 'Performing filesystem layer commit...' && echo ''
 
 # Create an Installed Build of the Engine
 WORKDIR [REDACTED_USER_HOME]/UnrealEngine
-RUN ./Engine/Build/BatchFiles/RunUAT.sh BuildGraph \
+RUN export BUILD_ID_OVERRIDE=UE_`cat ./Engine/Build/Build.version | jq --raw-output '"\(.MajorVersion).\(.MinorVersion)"'` && \
+	./Engine/Build/BatchFiles/RunUAT.sh BuildGraph \
 	-target="Make Installed Build Linux" \
 	-script=Engine/Build/InstalledEngineBuild.xml \
 	-set:HostPlatformOnly=true \
+	-set:WithClient=true \
 	-set:WithDDC=true \
-	-set:BuildIdOverride=UE_5.4 -set:WithLinuxArm64=false -set:WithClient=true -set:WithServer=true && \
+	-set:WithServer=true \
+	-set:BuildIdOverride="$BUILD_ID_OVERRIDE" -set:WithLinuxArm64=false -set:WithClient=true -set:WithServer=true && \
 	rm -R -f [REDACTED_USER_HOME]/UnrealEngine/LocalBuilds/InstalledDDC && echo '' && echo 'RUN directive complete. Docker will now commit the filesystem layer to disk.' && echo 'Note that for large filesystem layers this can take quite some time.' && echo 'Performing filesystem layer commit...' && echo ''
 
 # Ensure UnrealVersionSelector is built, since the prebuilt binaries may not be up-to-date
@@ -211,6 +229,11 @@ COPY --from=builder --chown=ue4:ue4 [REDACTED_USER_HOME]/UnrealEngine/Components/DDC /home/
 COPY --from=builder --chown=ue4:ue4 [REDACTED_USER_HOME]/.config/Epic/UnrealEngine/Install.ini [REDACTED_USER_HOME]/.config/Epic/UnrealEngine/Install.ini
 WORKDIR [REDACTED_USER_HOME]/UnrealEngine
 
+# GitHub Actions forcibly overrides $HOME and redirects it to a host directory that is bind-mounted at `/github/home`,
+# so we need to explicitly set XDG_CONFIG_HOME to ensure this does not bypass our engine config files and break things
+# (see: <[REDACTED_PROFILE_URL]/runner/issues/863>)
+ENV XDG_CONFIG_HOME=[REDACTED_USER_HOME]/.config
+
 # Add ushell to the system PATH and alias `ushell` to `ushell.sh`
 ENV PATH="$PATH:[REDACTED_USER_HOME]/UnrealEngine/Engine/Extras/ushell"
 RUN echo 'alias ushell="ushell.sh"' >> [REDACTED_USER_HOME]/.bashrc && echo '' && echo 'RUN directive complete. Docker will now commit the filesystem layer to disk.' && echo 'Note that for large filesystem layers this can take quite some time.' && echo 'Performing filesystem layer commit...' && echo ''
diff --git a/Engine/Extras/Containers/Dockerfiles/linux/dev/Dockerfile b/Engine/Extras/Containers/Dockerfiles/linux/dev/Dockerfile
--- a/Engine/Extras/Containers/Dockerfiles/linux/dev/Dockerfile
+++ b/Engine/Extras/Containers/Dockerfiles/linux/dev/Dockerfile
@@ -1,6 +1,7 @@
-# This file was generated by ue4-docker version 0.0.114 with the following options:
+# This file was generated by ue4-docker version 0.0.119 with the following options:
 # 
-# - buildgraph_args: "-set:BuildIdOverride=UE_5.4 -set:WithLinuxArm64=false -set:WithClient=true -set:WithServer=true"
+# - automatic_build_id: true
+# - buildgraph_args: "-set:WithLinuxArm64=false -set:WithClient=true -set:WithServer=true"
 # - combine: true
 # - credential_mode: "secrets"
 # - disable_all_patches: true
@@ -17,22 +18,22 @@
 # - ue4-minimal
 
 ARG BASEIMAGE
-FROM ${BASEIMAGE} as prerequisites
+FROM ${BASEIMAGE} AS prerequisites
 
 # Disable interactive prompts during package installation
 ENV DEBIAN_FRONTEND=noninteractive
 
 # Enable CUDA support for NVIDIA GPUs (even when not using a CUDA base image), since evidently some versions of UE unconditionally assume
 # `libcuda.so.1` exists when using the NVIDIA proprietary drivers, and will fail to initialise the Vulkan RHI if it is missing
-ENV NVIDIA_DRIVER_CAPABILITIES ${NVIDIA_DRIVER_CAPABILITIES},compute
+ENV NVIDIA_DRIVER_CAPABILITIES=${NVIDIA_DRIVER_CAPABILITIES},compute
 
 # Add the "display" driver capability for NVIDIA GPUs
 # (This allows us to run the Editor from an interactive container by bind-mounting the host system's X11 socket)
-ENV NVIDIA_DRIVER_CAPABILITIES ${NVIDIA_DRIVER_CAPABILITIES},display
+ENV NVIDIA_DRIVER_CAPABILITIES=${NVIDIA_DRIVER_CAPABILITIES},display
 
 # Enable NVENC support for use by Unreal Engine plugins that depend on it (e.g. Pixel Streaming)
 # (Note that adding `video` seems to implicitly enable `compute` as well, but we include separate directives here to clearly indicate the purpose of both)
-ENV NVIDIA_DRIVER_CAPABILITIES ${NVIDIA_DRIVER_CAPABILITIES},video
+ENV NVIDIA_DRIVER_CAPABILITIES=${NVIDIA_DRIVER_CAPABILITIES},video
 
 # Install our build prerequisites
 RUN apt-get update && apt-get install -y --no-install-recommends \
@@ -42,6 +43,7 @@ RUN apt-get update && apt-get install -y --no-install-recommends \
 		git \
 		git-lfs \
 		gpg-agent \
+		jq \
 		python3 \
 		python3-dev \
 		python3-pip \
@@ -87,6 +89,12 @@ RUN apt-get update && apt-get install -y --no-install-recommends \
 			xkb-data && \
 	rm -rf /var/lib/apt/lists/* && echo '' && echo 'RUN directive complete. Docker will now commit the filesystem layer to disk.' && echo 'Note that for large filesystem layers this can take quite some time.' && echo 'Performing filesystem layer commit...' && echo ''
 
+# Install the Avahi runtime libraries required by the NDIMedia plugin
+# (Starting in Unreal Engine 5.7, we need these installed before creating an Installed Build to prevent cooking failures)
+RUN apt-get update && apt-get install -y --no-install-recommends \
+			libavahi-client3 && \
+	rm -rf /var/lib/apt/lists/* && echo '' && echo 'RUN directive complete. Docker will now commit the filesystem layer to disk.' && echo 'Note that for large filesystem layers this can take quite some time.' && echo 'Performing filesystem layer commit...' && echo ''
+
 # Disable the default "lecture" message the first time a user runs a command using sudo
 RUN echo 'Defaults lecture="never"' >> /etc/sudoers && echo '' && echo 'RUN directive complete. Docker will now commit the filesystem layer to disk.' && echo 'Note that for large filesystem layers this can take quite some time.' && echo 'Performing filesystem layer commit...' && echo ''
 
@@ -150,29 +158,39 @@ RUN --mount=type=cache,target=[REDACTED_USER_HOME]/gitdeps,uid=1000,gid=1000 sudo apt-get u
 	./Setup.sh  && \
 	sudo rm -rf /var/lib/apt/lists/* && echo '' && echo 'RUN directive complete. Docker will now commit the filesystem layer to disk.' && echo 'Note that for large filesystem layers this can take quite some time.' && echo 'Performing filesystem layer commit...' && echo ''
 
-FROM source as builder
-
 # Set the changelist number in Build.version to ensure our Build ID is generated correctly
 ARG CHANGELIST
 COPY set-changelist.py /tmp/set-changelist.py
 RUN python3 /tmp/set-changelist.py [REDACTED_USER_HOME]/UnrealEngine/Engine/Build/Build.version $CHANGELIST && echo '' && echo 'RUN directive complete. Docker will now commit the filesystem layer to disk.' && echo 'Note that for large filesystem layers this can take quite some time.' && echo 'Performing filesystem layer commit...' && echo ''
 
+FROM source as builder
+
 # Remove the .git directory to disable UBT `git status` calls and speed up the build process
 RUN rm -rf [REDACTED_USER_HOME]/UnrealEngine/.git && echo '' && echo 'RUN directive complete. Docker will now commit the filesystem layer to disk.' && echo 'Note that for large filesystem layers this can take quite some time.' && echo 'Performing filesystem layer commit...' && echo ''
 
+# Disable Nuget package auditing (warnings NU1902 and NU1903)
+ENV NuGetAudit=false
+
 # Ensure UBT is built before we create the Installed Build, since Build.sh explicitly sets the
 # target .NET Framework version, whereas InstalledEngineBuild.xml just uses the system default,
 # which can result in errors when running the built UBT due to the wrong version being targeted
-RUN ./Engine/Build/BatchFiles/Linux/Build.sh ShaderCompileWorker Linux Development -SkipBuild -buildubt && echo '' && echo 'RUN directive complete. Docker will now commit the filesystem layer to disk.' && echo 'Note that for large filesystem layers this can take quite some time.' && echo 'Performing filesystem layer commit...' && echo ''
+RUN if [ -f ./Engine/Build/BatchFiles/BuildUBT.sh ]; then \
+		./Engine/Build/BatchFiles/BuildUBT.sh; \
+	else \
+		./Engine/Build/BatchFiles/Linux/Build.sh ShaderCompileWorker Linux Development -SkipBuild -buildubt; \
+	fi && echo '' && echo 'RUN directive complete. Docker will now commit the filesystem layer to disk.' && echo 'Note that for large filesystem layers this can take quite some time.' && echo 'Performing filesystem layer commit...' && echo ''
 
 # Create an Installed Build of the Engine
 WORKDIR [REDACTED_USER_HOME]/UnrealEngine
-RUN ./Engine/Build/BatchFiles/RunUAT.sh BuildGraph \
+RUN export BUILD_ID_OVERRIDE=UE_`cat ./Engine/Build/Build.version | jq --raw-output '"\(.MajorVersion).\(.MinorVersion)"'` && \
+	./Engine/Build/BatchFiles/RunUAT.sh BuildGraph \
 	-target="Make Installed Build Linux" \
 	-script=Engine/Build/InstalledEngineBuild.xml \
 	-set:HostPlatformOnly=true \
+	-set:WithClient=true \
 	-set:WithDDC=true \
-	-set:BuildIdOverride=UE_5.4 -set:WithLinuxArm64=false -set:WithClient=true -set:WithServer=true && \
+	-set:WithServer=true \
+	-set:BuildIdOverride="$BUILD_ID_OVERRIDE" -set:WithLinuxArm64=false -set:WithClient=true -set:WithServer=true && \
 	rm -R -f [REDACTED_USER_HOME]/UnrealEngine/LocalBuilds/InstalledDDC && echo '' && echo 'RUN directive complete. Docker will now commit the filesystem layer to disk.' && echo 'Note that for large filesystem layers this can take quite some time.' && echo 'Performing filesystem layer commit...' && echo ''
 
 # Ensure UnrealVersionSelector is built, since the prebuilt binaries may not be up-to-date
@@ -213,6 +231,11 @@ COPY --from=builder --chown=ue4:ue4 [REDACTED_USER_HOME]/UnrealEngine/Components/TemplatesA
 COPY --from=builder --chown=ue4:ue4 [REDACTED_USER_HOME]/.config/Epic/UnrealEngine/Install.ini [REDACTED_USER_HOME]/.config/Epic/UnrealEngine/Install.ini
 WORKDIR [REDACTED_USER_HOME]/UnrealEngine
 
+# GitHub Actions forcibly overrides $HOME and redirects it to a host directory that is bind-mounted at `/github/home`,
+# so we need to explicitly set XDG_CONFIG_HOME to ensure this does not bypass our engine config files and break things
+# (see: <[REDACTED_PROFILE_URL]/runner/issues/863>)
+ENV XDG_CONFIG_HOME=[REDACTED_USER_HOME]/.config
+
 # Add ushell to the system PATH and alias `ushell` to `ushell.sh`
 ENV PATH="$PATH:[REDACTED_USER_HOME]/UnrealEngine/Engine/Extras/ushell"
 RUN echo 'alias ushell="ushell.sh"' >> [REDACTED_USER_HOME]/.bashrc && echo '' && echo 'RUN directive complete. Docker will now commit the filesystem layer to disk.' && echo 'Note that for large filesystem layers this can take quite some time.' && echo 'Performing filesystem layer commit...' && echo ''
diff --git a/Engine/Extras/Containers/Dockerfiles/linux/generate.sh b/Engine/Extras/Containers/Dockerfiles/linux/generate.sh
--- a/Engine/Extras/Containers/Dockerfiles/linux/generate.sh
+++ b/Engine/Extras/Containers/Dockerfiles/linux/generate.sh
@@ -30,8 +30,11 @@ function generateDockerfile {
 		# This specifies the individual images whose Dockerfiles will be combined
 		--target=build-prerequisites,source,minimal
 		
-		# This sets the Build ID, disables building the Engine for AArch64, and enables support for client and dedicated server targets
-		--opt buildgraph-args='-set:BuildIdOverride=UE_5.4 -set:WithLinuxArm64=false -set:WithClient=true -set:WithServer=true'
+		# This disables building the Engine for AArch64, and enables support for client and dedicated server targets
+		--opt buildgraph-args='-set:WithLinuxArm64=false -set:WithClient=true -set:WithServer=true'
+		
+		# This automatically sets the Build ID so we don't need to specify it manually (requires ue4-docker version 0.0.119 or newer)
+		--opt automatic-build-id=true
 		
 		# This ensures we do not exclude binary files for any platforms
 		--opt gitdependencies-args=''
