Containers: improve error reporting when building engine container images

#rnx

diff --git a/Engine/Extras/Containers/Dockerfiles/linux/build.sh b/Engine/Extras/Containers/Dockerfiles/linux/build.sh
--- a/Engine/Extras/Containers/Dockerfiles/linux/build.sh
+++ b/Engine/Extras/Containers/Dockerfiles/linux/build.sh
@@ -1,17 +1,23 @@
 #!/usr/bin/env bash
 
-# Determine which release of the Unreal Engine we will be building container images for
-UNREAL_ENGINE_RELEASE="5.0"
-if [[ ! -z "$1" ]]; then
-	UNREAL_ENGINE_RELEASE="$1"
+# Determine which version of the Unreal Engine we will be building container images for
+UNREAL_ENGINE_VERSION="$1"
+if [[ -z "$UNREAL_ENGINE_VERSION" ]]; then
+	echo 'Usage:'
+	echo 'build.sh VERSION [REPOSITORY] [CHANGELIST] [TAG_SUFFIX]'
+	echo
+	echo 'Please specify an Unreal Engine version to build container images for.'
+	echo 'This can be a release number like `5.0.0`, or a Git branch/tag/commit.'
+	exit 1
 fi
 
-# Determine which Git branch/tag we will be pulling the Unreal Engine source code from
-# (If the release value contains a dash or does not contain two dots then we treat it as a branch/tag name)
-GIT_BRANCH="${UNREAL_ENGINE_RELEASE}-release"
-dots=$(echo "${UNREAL_ENGINE_RELEASE}" | grep -o '[.]' | wc -l)
-if [[ "${UNREAL_ENGINE_RELEASE}" == *"-"* ]] || [[ "${dots}" != "2" ]]; then
-	GIT_BRANCH="${UNREAL_ENGINE_RELEASE}"
+# Determine which Git branch/tag/commit we will be pulling the Unreal Engine source code from
+# (If the specified version value contains a dash or does not contain two dots then we treat it as a branch/tag/commit,
+#  otherwise we treat it as a release number and append `-release` to get the corresponding Git release tag)
+GIT_REF="${UNREAL_ENGINE_VERSION}-release"
+dots=$(echo "${UNREAL_ENGINE_VERSION}" | grep -o '[.]' | wc -l)
+if [[ "${UNREAL_ENGINE_VERSION}" == *"-"* ]] || [[ "${dots}" != "2" ]]; then
+	GIT_REF="${UNREAL_ENGINE_VERSION}"
 fi
 
 # If the user specified a custom Git repository to clone from then use that
@@ -26,34 +32,36 @@ if [[ ! -z "$3" ]]; then
 	CHANGELIST_OVERRIDE="$3"
 fi
 
-# If the user specified a git commit hash instead of a release, branch or tag then automatically determine the changelist
-# number from the commit message, and tag the built container images with an abbreviated version of the commit hash
-COMMIT_HASH=$(echo "${UNREAL_ENGINE_RELEASE}" | grep -E -o '^[0-9a-f]{40}$')
-COMMIT_SHORT=$(echo "${COMMIT_HASH}" | grep -E -o '^[0-9a-f]{8}')
-if [[ ! -z "$COMMIT_HASH" ]]; then
-	UNREAL_ENGINE_RELEASE="$COMMIT_SHORT"
-	CHANGELIST_OVERRIDE="auto"
+# Determine whether the user specified a custom tag suffix to use instead of the Unreal Engine version number
+TAG_SUFFIX="$UNREAL_ENGINE_VERSION"
+if [[ ! -z "$4" ]]; then
+	TAG_SUFFIX="$4"
 fi
 
 # Verify that the user has placed their GitHub username in the file `username.txt`
 if [ ! -f ./username.txt ]; then
+	echo 'Error: required credentials file missing!'
+	echo
 	echo 'Please place your GitHub username in a text file called `username.txt` in the current directory.'
 	exit 1
 fi
 
 # Verify that the user has placed their GitHub personal access token in the file `password.txt`
 if [ ! -f ./password.txt ]; then
+	echo 'Error: required credentials file missing!'
+	echo
 	echo 'Please place your GitHub personal access token in a text file called `password.txt` in the current directory.'
+	echo
 	echo 'Note that you should NOT use your GitHub password. See the following page for details on why a token is needed:'
 	echo 'https://github.blog/2020-12-15-token-authentication-requirements-for-git-operations/'
 	exit 1
 fi
 
 # Assemble the common arguments to pass to the `docker buildx build` command
-args=(
+commonArgs=(
 	--build-arg 'BASEIMAGE=nvidia/opengl:1.2-glvnd-devel-ubuntu22.04'
 	--build-arg "GIT_REPO=${GIT_REPO}"
-	--build-arg "GIT_BRANCH=${GIT_BRANCH}"
+	--build-arg "GIT_BRANCH=${GIT_REF}"
 	--build-arg "CHANGELIST=${CHANGELIST_OVERRIDE}"
 	--secret id=username,src=username.txt
 	--secret id=password,src=password.txt
@@ -61,39 +69,74 @@ args=(
 	--progress=plain
 )
 
+# Attempts to build a container image, and prints the log output if it fails
+function buildImage() {
+	local image="$1"
+	local context="$2"
+	local logfile="$3"
+	local args=("${@:4}")
+	echo "[docker buildx build -t $image" "${args[@]}" "$context]"
+	echo
+	docker buildx build -t "$image" "${args[@]}" "$context" &> "./logs/$logfile" || {
+		echo 'Error: image build failed! Log output:'
+		echo
+		cat "./logs/$logfile"
+		exit 1
+	}
+}
+
 # Print our build configuration values
 echo 'Build configuration'
 echo '-------------------'
 echo
-echo "Unreal Engine Release: ${UNREAL_ENGINE_RELEASE}"
+echo "Unreal Engine Version: ${UNREAL_ENGINE_VERSION}"
 echo "Git Repository:        ${GIT_REPO}"
-echo "Git Branch/Tag/Commit: ${GIT_BRANCH}"
-echo "Changelist Override:   ${CHANGELIST_OVERRIDE}"
+echo "Git Branch/Tag/Commit: ${GIT_REF}"
+echo "Changelist Override:   ${CHANGELIST_OVERRIDE:-(None)}"
+echo "Image Tag Suffix:      ${TAG_SUFFIX}"
 echo
 
-# Print commands as they are executed and halt immediately if any command fails
-set -ex
+# Halt immediately if any command fails
+set -e
+
+# Clean up any leftover logs from previous builds
+rm -rf ./logs
+mkdir ./logs
+
+# Build a runtime container image suitable for projects that do not use Pixel Streaming
+echo 'Building the runtime image...'
+buildImage \
+	'ghcr.io/epicgames/unreal-engine:runtime' \
+	./runtime \
+	runtime.log
 
-# Build our runtime container images
-docker build -t 'ghcr.io/epicgames/unreal-engine:runtime' ./runtime
-docker build -t 'ghcr.io/epicgames/unreal-engine:runtime-pixel-streaming' ./runtime-pixel-streaming
+# Build a runtime container image suitable for projects that use Pixel Streaming
+echo 'Building the Pixel Streaming runtime image...'
+buildImage \
+	'ghcr.io/epicgames/unreal-engine:runtime-pixel-streaming' \
+	./runtime-pixel-streaming \
+	runtime-pixel-streaming.log
 
 # Build a container image that encapsulates a full Installed Build of the Unreal Engine
-echo '[build.sh] Building the `unreal-engine` image for Unreal Engine release' "${UNREAL_ENGINE_RELEASE}..."
-docker buildx build \
-	-t "ghcr.io/epicgames/unreal-engine:dev-${UNREAL_ENGINE_RELEASE}" \
-	"${args[@]}" \
-	./dev
+echo 'Building the full `unreal-engine` image for Unreal Engine version' "${UNREAL_ENGINE_VERSION}..."
+buildImage \
+	"ghcr.io/epicgames/unreal-engine:dev-${TAG_SUFFIX}" \
+	./dev \
+	dev-full.log \
+	"${commonArgs[@]}"
 
 # Build a "slim" version of the container image that excludes debug symbols and templates
-echo '[build.sh] Building the `unreal-engine` image for Unreal Engine release' "${UNREAL_ENGINE_RELEASE}..."
-docker buildx build \
-	-t "ghcr.io/epicgames/unreal-engine:dev-slim-${UNREAL_ENGINE_RELEASE}" \
-	"${args[@]}" \
-	./dev-slim
+echo 'Building the slim `unreal-engine` image for Unreal Engine version' "${UNREAL_ENGINE_VERSION}..."
+buildImage \
+	"ghcr.io/epicgames/unreal-engine:dev-slim-${TAG_SUFFIX}" \
+	./dev-slim \
+	dev-slim.log \
+	"${commonArgs[@]}"
 
 # Build the container image for the Multi-User Editing server, which pulls files from our other images
-docker build \
-	-t "ghcr.io/epicgames/multi-user-server:${UNREAL_ENGINE_RELEASE}" \
-	--build-arg "UNREAL_ENGINE_RELEASE=${UNREAL_ENGINE_RELEASE}" \
-	./multi-user-server
+echo 'Building the `multi-user-server` image for Unreal Engine version' "${UNREAL_ENGINE_VERSION}..."
+buildImage \
+	"ghcr.io/epicgames/multi-user-server:${TAG_SUFFIX}" \
+	./multi-user-server \
+	multi-user-server.log \
+	--build-arg "UNREAL_ENGINE_VERSION=${TAG_SUFFIX}"
diff --git a/Engine/Extras/Containers/Dockerfiles/linux/multi-user-server/Dockerfile b/Engine/Extras/Containers/Dockerfiles/linux/multi-user-server/Dockerfile
--- a/Engine/Extras/Containers/Dockerfiles/linux/multi-user-server/Dockerfile
+++ b/Engine/Extras/Containers/Dockerfiles/linux/multi-user-server/Dockerfile
@@ -1,5 +1,5 @@
-ARG UNREAL_ENGINE_RELEASE
-FROM ghcr.io/epicgames/unreal-engine:dev-${UNREAL_ENGINE_RELEASE} as source
+ARG UNREAL_ENGINE_VERSION
+FROM ghcr.io/epicgames/unreal-engine:dev-${UNREAL_ENGINE_VERSION} as source
 
 # Gather the required files for the UnrealMultiUserServer binary and its dependencies
 COPY copy-target.py /tmp/copy-target.py
