Containers: add support for specifying a custom log directory when building engine container images

#rnx

diff --git a/Engine/Extras/Containers/Dockerfiles/linux/build.sh b/Engine/Extras/Containers/Dockerfiles/linux/build.sh
--- a/Engine/Extras/Containers/Dockerfiles/linux/build.sh
+++ b/Engine/Extras/Containers/Dockerfiles/linux/build.sh
@@ -38,6 +38,12 @@ if [[ ! -z "$4" ]]; then
 	TAG_SUFFIX="$4"
 fi
 
+# If a custom log directory was specified then use it
+LOG_DIRECTORY="./logs"
+if [[ ! -z "$CONTAINER_BUILD_LOGS" ]]; then
+	LOG_DIRECTORY="$CONTAINER_BUILD_LOGS"
+fi
+
 # Verify that the user has placed their GitHub username in the file `username.txt`
 if [ ! -f ./username.txt ]; then
 	echo 'Error: required credentials file missing!'
@@ -77,10 +83,10 @@ function buildImage() {
 	local args=("${@:4}")
 	echo "[docker buildx build -t $image" "${args[@]}" "$context]"
 	echo
-	docker buildx build -t "$image" "${args[@]}" "$context" &> "./logs/$logfile" || {
+	docker buildx build -t "$image" "${args[@]}" "$context" &> "$LOG_DIRECTORY/$logfile" || {
 		echo 'Error: image build failed! Log output:'
 		echo
-		cat "./logs/$logfile"
+		cat "$LOG_DIRECTORY/$logfile"
 		exit 1
 	}
 }
@@ -99,9 +105,9 @@ echo
 # Halt immediately if any command fails
 set -e
 
-# Clean up any leftover logs from previous builds
-rm -rf ./logs
-mkdir ./logs
+# Ensure the log directory exists and clean up any leftover logs from previous builds
+test -d "$LOG_DIRECTORY" || mkdir -p "$LOG_DIRECTORY"
+rm -f "$LOG_DIRECTORY"/*.log
 
 # Build a runtime container image suitable for projects that do not use Pixel Streaming
 echo 'Building the runtime image...'
